NVIDIA Reflex integration into ShooterGame.

#rnx

diff --git a/Engine/Plugins/Runtime/Nvidia/Reflex/Source/Private/ReflexLatencyMarkers.cpp b/Engine/Plugins/Runtime/Nvidia/Reflex/Source/Private/ReflexLatencyMarkers.cpp
--- a/Engine/Plugins/Runtime/Nvidia/Reflex/Source/Private/ReflexLatencyMarkers.cpp
+++ b/Engine/Plugins/Runtime/Nvidia/Reflex/Source/Private/ReflexLatencyMarkers.cpp
@@ -5,6 +5,9 @@
 #include "HAL/IConsoleManager.h"
 
 #include "RHI.h"
+#include "Engine/Console.h"
+#include "GameFramework/PlayerController.h"
+#include "Kismet/GameplayStatics.h"
 
 #include "Windows/AllowWindowsPlatformTypes.h"
 #include <D3D11.h>
@@ -34,12 +37,24 @@ void FReflexLatencyMarkers::Initialize()
 		{
 			bProperDriverVersion = true;
 		}
+
+		// Need to verify that Reflex Low Latency mode is supported on current hardware
+		NV_GET_SLEEP_STATUS_PARAMS_V1 SleepStatusParams = { 0 };
+		SleepStatusParams.version = NV_GET_SLEEP_STATUS_PARAMS_VER1;
+
+		NvAPI_Status status = NVAPI_OK;
+		status = NvAPI_D3D_GetSleepStatus(static_cast<IUnknown*>(GDynamicRHI->RHIGetNativeDevice()), &SleepStatusParams);
+
+		if (status == NVAPI_OK)
+		{
+			bFeatureSupport = true;
+		}
 	}
 }
 
 void FReflexLatencyMarkers::Tick(float DeltaTime)
 {
-	if (DisableLatencyMarkers == 0 && bEnabled && bProperDriverVersion && IsRHIDeviceNVIDIA())
+	if (DisableLatencyMarkers == 0 && bEnabled && bProperDriverVersion && IsRHIDeviceNVIDIA() && bFeatureSupport)
 	{
 		NvAPI_Status LatencyStatus = NVAPI_OK;
 		NV_LATENCY_RESULT_PARAMS_V1 LatencyResults = { 0 };
@@ -57,31 +72,55 @@ void FReflexLatencyMarkers::Tick(float DeltaTime)
 
 				// A 3/4, 1/4 split gets close to a simple 10 frame moving average
 				AverageTotalLatencyMs = AverageTotalLatencyMs * 0.75f + TotalLatencyUs / 1000.0f * 0.25f;
-				AverageGameLatencyMs = AverageGameLatencyMs * 0.75f + (LatencyResults.frameReport[63].simEndTime - LatencyResults.frameReport[63].simStartTime) / 1000.0f * 0.25f;
-				AverageRenderLatencyMs = AverageRenderLatencyMs * 0.75f + (LatencyResults.frameReport[63].renderSubmitEndTime - LatencyResults.frameReport[63].renderSubmitStartTime) / 1000.0f * 0.25f;
+
+				AverageGameLatencyMs = AverageGameLatencyMs * 0.75f + (LatencyResults.frameReport[63].driverEndTime - LatencyResults.frameReport[63].simStartTime) / 1000.0f * 0.25f;
+				AverageRenderLatencyMs = AverageRenderLatencyMs * 0.75f + (LatencyResults.frameReport[63].gpuRenderEndTime - LatencyResults.frameReport[63].osRenderQueueStartTime) / 1000.0f * 0.25f;
+
+				AverageSimulationLatencyMs = AverageSimulationLatencyMs * 0.75f + (LatencyResults.frameReport[63].simEndTime - LatencyResults.frameReport[63].simStartTime) / 1000.0f * 0.25f;
+				AverageRenderSubmitLatencyMs = AverageRenderSubmitLatencyMs * 0.75f + (LatencyResults.frameReport[63].renderSubmitEndTime - LatencyResults.frameReport[63].renderSubmitStartTime) / 1000.0f * 0.25f;
+				AveragePresentLatencyMs = AveragePresentLatencyMs * 0.75f + (LatencyResults.frameReport[63].presentEndTime - LatencyResults.frameReport[63].presentStartTime) / 1000.0f * 0.25f;
 				AverageDriverLatencyMs = AverageDriverLatencyMs * 0.75f + (LatencyResults.frameReport[63].driverEndTime - LatencyResults.frameReport[63].driverStartTime) / 1000.0f * 0.25f;
-				AverageOSWorkQueueLatencyMs = AverageOSWorkQueueLatencyMs * 0.75f + (LatencyResults.frameReport[63].osRenderQueueEndTime - LatencyResults.frameReport[63].osRenderQueueStartTime) / 1000.0f * 0.25f;
+				AverageOSRenderQueueLatencyMs = AverageOSRenderQueueLatencyMs * 0.75f + (LatencyResults.frameReport[63].osRenderQueueEndTime - LatencyResults.frameReport[63].osRenderQueueStartTime) / 1000.0f * 0.25f;
 				AverageGPURenderLatencyMs = AverageGPURenderLatencyMs * 0.75f + (LatencyResults.frameReport[63].gpuRenderEndTime - LatencyResults.frameReport[63].gpuRenderStartTime) / 1000.0f * 0.25f;
 
-				RenderOffsetMs = (LatencyResults.frameReport[63].renderSubmitStartTime - LatencyResults.frameReport[63].simStartTime) / 1000.0f;
+				RenderSubmitOffsetMs = (LatencyResults.frameReport[63].renderSubmitStartTime - LatencyResults.frameReport[63].simStartTime) / 1000.0f;
+				PresentOffsetMs = (LatencyResults.frameReport[63].presentStartTime - LatencyResults.frameReport[63].simStartTime) / 1000.0f;
 				DriverOffsetMs = (LatencyResults.frameReport[63].driverStartTime - LatencyResults.frameReport[63].simStartTime) / 1000.0f;
-				OSWorkQueueOffsetMs = (LatencyResults.frameReport[63].osRenderQueueStartTime - LatencyResults.frameReport[63].simStartTime) / 1000.0f;
+				OSRenderQueueOffsetMs = (LatencyResults.frameReport[63].osRenderQueueStartTime - LatencyResults.frameReport[63].simStartTime) / 1000.0f;
 				GPURenderOffsetMs = (LatencyResults.frameReport[63].gpuRenderStartTime - LatencyResults.frameReport[63].simStartTime) / 1000.0f;
 
- 				UE_LOG(LogLatencyMarkers, VeryVerbose, TEXT("AverageTotalLatencyMs: %f"), AverageTotalLatencyMs);
+				UE_LOG(LogLatencyMarkers, VeryVerbose, TEXT("AverageTotalLatencyMs: %f"), AverageTotalLatencyMs);
 				UE_LOG(LogLatencyMarkers, VeryVerbose, TEXT("AverageGameLatencyMs: %f"), AverageGameLatencyMs);
- 				UE_LOG(LogLatencyMarkers, VeryVerbose, TEXT("AverageRenderLatencyMs: %f"), AverageRenderLatencyMs);
- 				UE_LOG(LogLatencyMarkers, VeryVerbose, TEXT("AverageDriverLatencyMs: %f"), AverageDriverLatencyMs);
- 				UE_LOG(LogLatencyMarkers, VeryVerbose, TEXT("AverageOSWorkQueueLatencyMs: %f"), AverageOSWorkQueueLatencyMs);
- 				UE_LOG(LogLatencyMarkers, VeryVerbose, TEXT("AverageGPURenderLatencyMs: %f"), AverageGPURenderLatencyMs);
+				UE_LOG(LogLatencyMarkers, VeryVerbose, TEXT("AverageRenderLatencyMs: %f"), AverageRenderLatencyMs);
+
+				UE_LOG(LogLatencyMarkers, VeryVerbose, TEXT("AverageSimulationLatencyMs: %f"), AverageSimulationLatencyMs);
+				UE_LOG(LogLatencyMarkers, VeryVerbose, TEXT("AverageRenderSubmitLatencyMs: %f"), AverageRenderSubmitLatencyMs);
+				UE_LOG(LogLatencyMarkers, VeryVerbose, TEXT("AveragePresentLatencyMs: %f"), AveragePresentLatencyMs);
+				UE_LOG(LogLatencyMarkers, VeryVerbose, TEXT("AverageDriverLatencyMs: %f"), AverageDriverLatencyMs);
+				UE_LOG(LogLatencyMarkers, VeryVerbose, TEXT("AverageOSRenderQueueLatencyMs: %f"), AverageOSRenderQueueLatencyMs);
+				UE_LOG(LogLatencyMarkers, VeryVerbose, TEXT("AverageGPURenderLatencyMs: %f"), AverageGPURenderLatencyMs);
 			}
 		}
 	}
 }
 
-void FReflexLatencyMarkers::SetGameLatencyMarkerStart(uint64 FrameNumber)
+void FReflexLatencyMarkers::SetInputSampleLatencyMarker(uint64 FrameNumber)
+{
+	if (DisableLatencyMarkers == 0 && bProperDriverVersion && bEnabled && IsRHIDeviceNVIDIA() && bFeatureSupport)
+	{
+		NvAPI_Status status = NVAPI_OK;
+		NV_LATENCY_MARKER_PARAMS_V1 params = { 0 };
+		params.version = NV_LATENCY_MARKER_PARAMS_VER1;
+		params.frameID = FrameNumber;
+		params.markerType = INPUT_SAMPLE;
+
+		status = NvAPI_D3D_SetLatencyMarker(static_cast<IUnknown*>(GDynamicRHI->RHIGetNativeDevice()), &params);
+	}
+}
+
+void FReflexLatencyMarkers::SetSimulationLatencyMarkerStart(uint64 FrameNumber)
 {
-	if (DisableLatencyMarkers == 0 && bProperDriverVersion && bEnabled && IsRHIDeviceNVIDIA())
+	if (DisableLatencyMarkers == 0 && bProperDriverVersion && bEnabled && IsRHIDeviceNVIDIA() && bFeatureSupport)
 	{
 		NvAPI_Status status = NVAPI_OK;
 		NV_LATENCY_MARKER_PARAMS_V1 params = { 0 };
@@ -93,9 +132,9 @@ void FReflexLatencyMarkers::SetGameLatencyMarkerStart(uint64 FrameNumber)
 	}
 }
 
-void FReflexLatencyMarkers::SetGameLatencyMarkerEnd(uint64 FrameNumber)
+void FReflexLatencyMarkers::SetSimulationLatencyMarkerEnd(uint64 FrameNumber)
 {
-	if (DisableLatencyMarkers == 0 && bProperDriverVersion && bEnabled && IsRHIDeviceNVIDIA())
+	if (DisableLatencyMarkers == 0 && bProperDriverVersion && bEnabled && IsRHIDeviceNVIDIA() && bFeatureSupport)
 	{
 		NvAPI_Status status = NVAPI_OK;
 		NV_LATENCY_MARKER_PARAMS_V1 params = { 0 };
@@ -107,9 +146,9 @@ void FReflexLatencyMarkers::SetGameLatencyMarkerEnd(uint64 FrameNumber)
 	}
 }
 
-void FReflexLatencyMarkers::SetRenderLatencyMarkerStart(uint64 FrameNumber)
+void FReflexLatencyMarkers::SetRenderSubmitLatencyMarkerStart(uint64 FrameNumber)
 {
-	if (DisableLatencyMarkers == 0 && bProperDriverVersion && bEnabled && IsRHIDeviceNVIDIA())
+	if (DisableLatencyMarkers == 0 && bProperDriverVersion && bEnabled && IsRHIDeviceNVIDIA() && bFeatureSupport)
 	{
 		NvAPI_Status status = NVAPI_OK;
 		NV_LATENCY_MARKER_PARAMS_V1 params = { 0 };
@@ -121,9 +160,9 @@ void FReflexLatencyMarkers::SetRenderLatencyMarkerStart(uint64 FrameNumber)
 	}
 }
 
-void FReflexLatencyMarkers::SetRenderLatencyMarkerEnd(uint64 FrameNumber)
+void FReflexLatencyMarkers::SetRenderSubmitLatencyMarkerEnd(uint64 FrameNumber)
 {
-	if (DisableLatencyMarkers == 0 && bProperDriverVersion && bEnabled && IsRHIDeviceNVIDIA())
+	if (DisableLatencyMarkers == 0 && bProperDriverVersion && bEnabled && IsRHIDeviceNVIDIA() && bFeatureSupport)
 	{
 		NvAPI_Status status = NVAPI_OK;
 		NV_LATENCY_MARKER_PARAMS_V1 params = { 0 };
@@ -135,30 +174,162 @@ void FReflexLatencyMarkers::SetRenderLatencyMarkerEnd(uint64 FrameNumber)
 	}
 }
 
-void FReflexLatencyMarkers::SetCustomLatencyMarker(uint32 MarkerId, uint64 FrameNumber)
+void FReflexLatencyMarkers::SetPresentLatencyMarkerStart(uint64 FrameNumber)
 {
-	// Allow trigger flash here
-	if (NV_LATENCY_MARKER_TYPE(MarkerId) == TRIGGER_FLASH)
+	if (DisableLatencyMarkers == 0 && bProperDriverVersion && bEnabled && IsRHIDeviceNVIDIA() && bFeatureSupport)
 	{
-		if (DisableLatencyMarkers == 0 && bProperDriverVersion && bEnabled && IsRHIDeviceNVIDIA())
+		NvAPI_Status status = NVAPI_OK;
+		NV_LATENCY_MARKER_PARAMS_V1 params = { 0 };
+		params.version = NV_LATENCY_MARKER_PARAMS_VER1;
+		params.frameID = FrameNumber;
+		params.markerType = PRESENT_START;
+
+		status = NvAPI_D3D_SetLatencyMarker(static_cast<IUnknown*>(GDynamicRHI->RHIGetNativeDevice()), &params);
+	}
+}
+
+void FReflexLatencyMarkers::SetPresentLatencyMarkerEnd(uint64 FrameNumber)
+{
+	if (DisableLatencyMarkers == 0 && bProperDriverVersion && bEnabled && IsRHIDeviceNVIDIA() && bFeatureSupport)
+	{
+		NvAPI_Status status = NVAPI_OK;
+		NV_LATENCY_MARKER_PARAMS_V1 params = { 0 };
+		params.version = NV_LATENCY_MARKER_PARAMS_VER1;
+		params.frameID = FrameNumber;
+		params.markerType = PRESENT_END;
+
+		status = NvAPI_D3D_SetLatencyMarker(static_cast<IUnknown*>(GDynamicRHI->RHIGetNativeDevice()), &params);
+	}
+}
+
+void FReflexLatencyMarkers::SetFlashIndicatorLatencyMarker(uint64 FrameNumber)
+{
+	if (DisableLatencyMarkers == 0 && bProperDriverVersion && bEnabled && IsRHIDeviceNVIDIA() && bFeatureSupport)
+	{
+		if (bFlashIndicatorEnabled)
 		{
 			NvAPI_Status status = NVAPI_OK;
 			NV_LATENCY_MARKER_PARAMS_V1 params = { 0 };
 			params.version = NV_LATENCY_MARKER_PARAMS_VER1;
 			params.frameID = FrameNumber;
-			params.markerType = NV_LATENCY_MARKER_TYPE(MarkerId);
+			params.markerType = TRIGGER_FLASH;
 
 			status = NvAPI_D3D_SetLatencyMarker(static_cast<IUnknown*>(GDynamicRHI->RHIGetNativeDevice()), &params);
 		}
 	}
 }
 
+void FReflexLatencyMarkers::SetCustomLatencyMarker(uint32 MarkerId, uint64 FrameNumber)
+{
+	if (DisableLatencyMarkers == 0 && bProperDriverVersion && bEnabled && IsRHIDeviceNVIDIA() && bFeatureSupport)
+	{
+		NvAPI_Status status = NVAPI_OK;
+		NV_LATENCY_MARKER_PARAMS_V1 params = { 0 };
+		params.version = NV_LATENCY_MARKER_PARAMS_VER1;
+		params.frameID = FrameNumber;
+		params.markerType = NV_LATENCY_MARKER_TYPE(MarkerId);
+
+		status = NvAPI_D3D_SetLatencyMarker(static_cast<IUnknown*>(GDynamicRHI->RHIGetNativeDevice()), &params);
+	}
+}
+
+void FReflexLatencyMarkers::SetFlashIndicatorEnabled(bool bInEnabled)
+{
+	bFlashIndicatorEnabled = bInEnabled;
+}
+
+bool FReflexLatencyMarkers::GetFlashIndicatorEnabled()
+{
+	if (DisableLatencyMarkers == 1 || !bProperDriverVersion || !IsRHIDeviceNVIDIA() || !bFeatureSupport)
+	{
+		return false;
+	}
+
+	return bFlashIndicatorEnabled;
+}
+
+void FReflexLatencyMarkers::SetEnabled(bool bInEnabled)
+{
+	if (!bInEnabled)
+	{
+		// Reset module back to default values in case re-enabled in the same session
+		AverageTotalLatencyMs = 0.0f;
+		AverageGameLatencyMs = 0.0f;
+		AverageRenderLatencyMs = 0.0f;
+
+		AverageSimulationLatencyMs = 0.0f;
+		AverageRenderSubmitLatencyMs = 0.0f;
+		AveragePresentLatencyMs = 0.0f;
+		AverageDriverLatencyMs = 0.0f;
+		AverageOSRenderQueueLatencyMs = 0.0f;
+		AverageGPURenderLatencyMs = 0.0f;
+
+		RenderSubmitOffsetMs = 0.0f;
+		PresentOffsetMs = 0.0f;
+		DriverOffsetMs = 0.0f;
+		OSRenderQueueOffsetMs = 0.0f;
+		GPURenderOffsetMs = 0.0f;
+
+		bFlashIndicatorEnabled = false;
+	}
+
+	bEnabled = bInEnabled;
+}
+
 bool FReflexLatencyMarkers::GetEnabled()
 {
-	if (DisableLatencyMarkers == 1 || !bProperDriverVersion || !IsRHIDeviceNVIDIA())
+	if (DisableLatencyMarkers == 1 || !bProperDriverVersion || !IsRHIDeviceNVIDIA() || !bFeatureSupport)
 	{
 		return false;
 	}
 
 	return bEnabled;
+}
+
+bool FReflexLatencyMarkers::GetAvailable()
+{
+	if (DisableLatencyMarkers == 1 || !bProperDriverVersion || !IsRHIDeviceNVIDIA() || !bFeatureSupport)
+	{
+		return false;
+	}
+
+	return true;
+}
+
+bool FReflexLatencyMarkers::Exec(UWorld* InWorld, const TCHAR* Cmd, FOutputDevice& Ar)
+{
+	bool bHandled = false;
+	FString ReflexLatencyMarkers;
+	APlayerController* PlayerController = (InWorld ? UGameplayStatics::GetPlayerController(InWorld, 0) : NULL);
+	ULocalPlayer* LocalPlayer = (PlayerController ? Cast<ULocalPlayer>(PlayerController->Player) : NULL);
+
+	if (FParse::Value(Cmd, TEXT("ReflexLatencyMarkers="), ReflexLatencyMarkers))
+	{
+		if (ReflexLatencyMarkers == "0")
+		{
+			SetEnabled(false);
+
+			if (LocalPlayer && LocalPlayer->ViewportClient && LocalPlayer->ViewportClient->ViewportConsole)
+			{
+				LocalPlayer->ViewportClient->ViewportConsole->OutputText("Reflex Latency Markers: Off");
+			}
+
+			UE_LOG(LogLatencyMarkers, Log, TEXT("Reflex Latency Markers: Off"));
+		}
+		else if (ReflexLatencyMarkers == "1")
+		{
+			SetEnabled(true);
+
+			if (LocalPlayer && LocalPlayer->ViewportClient && LocalPlayer->ViewportClient->ViewportConsole)
+			{
+				LocalPlayer->ViewportClient->ViewportConsole->OutputText("Reflex Latency Markers: On");
+			}
+
+			UE_LOG(LogLatencyMarkers, Log, TEXT("Reflex Latency Markers: On"));
+		}
+
+		bHandled = true;
+	}
+
+	return bHandled;
 }
\ No newline at end of file
diff --git a/Engine/Plugins/Runtime/Nvidia/Reflex/Source/Private/ReflexMaxTickRateHandler.cpp b/Engine/Plugins/Runtime/Nvidia/Reflex/Source/Private/ReflexMaxTickRateHandler.cpp
--- a/Engine/Plugins/Runtime/Nvidia/Reflex/Source/Private/ReflexMaxTickRateHandler.cpp
+++ b/Engine/Plugins/Runtime/Nvidia/Reflex/Source/Private/ReflexMaxTickRateHandler.cpp
@@ -5,6 +5,9 @@
 #include "HAL/IConsoleManager.h"
 
 #include "RHI.h"
+#include "Engine/Console.h"
+#include "GameFramework/PlayerController.h"
+#include "Kismet/GameplayStatics.h"
 
 #include "Windows/AllowWindowsPlatformTypes.h"
 #include <D3D11.h>
@@ -23,7 +26,7 @@ int32 DisableLatencyMarkerOptimize = 0;
 static FAutoConsoleVariableRef CVarDisableLatencyMarkerOptimize(
 	TEXT("t.DisableLatencyMarkerOptimize"),
 	DisableLatencyMarkerOptimize,
-	TEXT("Disable Disable Latency Marker Optimize")
+	TEXT("Disable Latency Marker Optimize")
 );
 
 DEFINE_LOG_CATEGORY_STATIC(LogMaxTickRateHandler, Log, All);
@@ -41,12 +44,24 @@ void FReflexMaxTickRateHandler::Initialize()
 		{
 			bProperDriverVersion = true;
 		}
+
+		// Need to verify that Reflex Low Latency mode is supported on current hardware
+		NV_GET_SLEEP_STATUS_PARAMS_V1 SleepStatusParams = { 0 };
+		SleepStatusParams.version = NV_GET_SLEEP_STATUS_PARAMS_VER1;
+
+		NvAPI_Status status = NVAPI_OK;
+		status = NvAPI_D3D_GetSleepStatus(static_cast<IUnknown*>(GDynamicRHI->RHIGetNativeDevice()), &SleepStatusParams);
+
+		if (status == NVAPI_OK)
+		{
+			bFeatureSupport = true;
+		}
 	}
 }
 
 bool FReflexMaxTickRateHandler::HandleMaxTickRate(float DesiredMaxTickRate)
 {
-	if (DisableCustomTickRateHandler == 0 && bProperDriverVersion && !GIsEditor && IsRHIDeviceNVIDIA())
+	if (DisableCustomTickRateHandler == 0 && bProperDriverVersion && !GIsEditor && IsRHIDeviceNVIDIA() && bFeatureSupport)
 	{
 		if (bEnabled)
 		{
@@ -56,25 +71,20 @@ bool FReflexMaxTickRateHandler::HandleMaxTickRate(float DesiredMaxTickRate)
 				NvAPI_Status status = NVAPI_OK;
 				NV_SET_SLEEP_MODE_PARAMS_V1 params = { 0 };
 				params.version = NV_SET_SLEEP_MODE_PARAMS_VER1;
-				params.bLowLatencyMode = bUltraLowLatency;
-				params.bLowLatencyBoost = bGPUBoost;
+				params.bLowLatencyMode = bLowLatencyMode;
+				params.bLowLatencyBoost = bBoost;
 				MinimumInterval = DesiredMinimumInterval;
 				params.minimumIntervalUs = MinimumInterval;
 				params.bUseMarkersToOptimize = DisableLatencyMarkerOptimize ? NV_FALSE : NV_TRUE;
 				status = NvAPI_D3D_SetSleepMode(static_cast<IUnknown*>(GDynamicRHI->RHIGetNativeDevice()), &params);
 
-				UE_LOG(LogMaxTickRateHandler, Log, TEXT("SetSleepMode MI:%f L:%s B:%s"), MinimumInterval, bUltraLowLatency ? TEXT("On") : TEXT("Off"), bGPUBoost ? TEXT("On") : TEXT("Off"));
-
-				// Need to verify that Low Latency flag actually applied, NVIDIA says the return code can still be good, but it might be incompatible
-				NV_GET_SLEEP_STATUS_PARAMS_V1 SleepStatusParams = { 0 };
-				SleepStatusParams.version = NV_GET_SLEEP_STATUS_PARAMS_VER1;
-				status = NvAPI_D3D_GetSleepStatus(static_cast<IUnknown*>(GDynamicRHI->RHIGetNativeDevice()), &SleepStatusParams);
-				if (bUltraLowLatency && SleepStatusParams.bLowLatencyMode == false)
+				// Need to verify that Low Latency mode change actually applied successfully
+				if (bLowLatencyMode && (status != NVAPI_OK))
 				{
 					UE_LOG(LogMaxTickRateHandler, Warning, TEXT("Unable to turn on low latency"));
 					// Clear the ULL flag
 					CustomFlags = CustomFlags & ~1;
-					bUltraLowLatency = false;
+					bLowLatencyMode = false;
 				}
 
 				LastCustomFlags = CustomFlags;
@@ -91,17 +101,22 @@ bool FReflexMaxTickRateHandler::HandleMaxTickRate(float DesiredMaxTickRate)
 			// When disabled, if we ever called SetSleepMode, we need to clean up after ourselves
 			if (bWasEnabled)
 			{
-				LastCustomFlags = 0;
-				bWasEnabled = false;
-
 				NvAPI_Status status = NVAPI_OK;
 				NV_SET_SLEEP_MODE_PARAMS_V1 params = { 0 };
 				params.version = NV_SET_SLEEP_MODE_PARAMS_VER1;
 				params.bLowLatencyMode = false;
 				params.bLowLatencyBoost = false;
 				params.minimumIntervalUs = 0;
+				params.bUseMarkersToOptimize = false;
 				status = NvAPI_D3D_SetSleepMode(static_cast<IUnknown*>(GDynamicRHI->RHIGetNativeDevice()), &params);
 				UE_LOG(LogMaxTickRateHandler, Log, TEXT("SetSleepMode clean up"));
+
+				// Reset module back to default values in case re-enabled in the same session
+				MinimumInterval = -1.0f;
+				LastCustomFlags = 0;
+				bWasEnabled = false;
+				bLowLatencyMode = true;
+				bBoost = false;
 			}
 		}
 	}
@@ -114,23 +129,129 @@ void FReflexMaxTickRateHandler::SetFlags(uint32 Flags)
 	CustomFlags = Flags;
 	if ((Flags & 1) > 0)
 	{
-		bUltraLowLatency = true;
+		bLowLatencyMode = true;
 	}
 	else
 	{
-		bUltraLowLatency = false;
+		bLowLatencyMode = false;
 	}
 	if ((Flags & 2) > 0)
 	{
-		bGPUBoost = true;
+		bBoost = true;
 	}
 	else
 	{
-		bGPUBoost = false;
+		bBoost = false;
 	}
 }
 
 uint32 FReflexMaxTickRateHandler::GetFlags()
 {
-	return CustomFlags;	
+	return CustomFlags;
+}
+
+void FReflexMaxTickRateHandler::SetEnabled(bool bInEnabled)
+{
+	bEnabled = bInEnabled;
+}
+
+bool FReflexMaxTickRateHandler::GetEnabled()
+{
+	if (DisableCustomTickRateHandler == 1 || !bProperDriverVersion || GIsEditor || !IsRHIDeviceNVIDIA() || !bFeatureSupport)
+	{
+		return false;
+	}
+
+	return bEnabled;
+}
+
+bool FReflexMaxTickRateHandler::GetAvailable()
+{
+	if (DisableCustomTickRateHandler == 1 || !bProperDriverVersion || GIsEditor || !IsRHIDeviceNVIDIA() || !bFeatureSupport)
+	{
+		return false;
+	}
+
+	return true;
+}
+
+bool FReflexMaxTickRateHandler::Exec(UWorld* InWorld, const TCHAR* Cmd, FOutputDevice& Ar)
+{
+	bool bHandled = false;
+	FString ReflexMode;
+	APlayerController* PlayerController = (InWorld ? UGameplayStatics::GetPlayerController(InWorld, 0) : NULL);
+	ULocalPlayer* LocalPlayer = (PlayerController ? Cast<ULocalPlayer>(PlayerController->Player) : NULL);
+
+	if (FParse::Value(Cmd, TEXT("ReflexMode="), ReflexMode))
+	{
+		if (ReflexMode == "0")
+		{
+			SetEnabled(false);
+			
+			if (LocalPlayer && LocalPlayer->ViewportClient && LocalPlayer->ViewportClient->ViewportConsole)
+			{
+				LocalPlayer->ViewportClient->ViewportConsole->OutputText("Reflex Low Latency mode: Off");
+			}
+
+			UE_LOG(LogMaxTickRateHandler, Log, TEXT("Reflex Low Latency mode: Off"));
+		}
+		else if (ReflexMode == "1")
+		{
+			SetEnabled(true);
+			SetFlags(1);
+
+			if (LocalPlayer && LocalPlayer->ViewportClient && LocalPlayer->ViewportClient->ViewportConsole)
+			{
+				LocalPlayer->ViewportClient->ViewportConsole->OutputText("Reflex Low Latency mode: On");
+			}
+
+			UE_LOG(LogMaxTickRateHandler, Log, TEXT("Reflex Low Latency mode: On"));
+		}
+		else if (ReflexMode == "2")
+		{
+			SetEnabled(true);
+			SetFlags(3);
+
+			if (LocalPlayer && LocalPlayer->ViewportClient && LocalPlayer->ViewportClient->ViewportConsole)
+			{
+				LocalPlayer->ViewportClient->ViewportConsole->OutputText("Reflex Low Latency mode: On+Boost");
+			}
+
+			UE_LOG(LogMaxTickRateHandler, Log, TEXT("Reflex Low Latency mode: On+Boost"));
+		}
+		
+		bHandled = true;
+	}
+	else if (FParse::Command(&Cmd, TEXT("ReflexModeToggle")))
+	{
+		bool bLowLatencyModeEnabled = GetEnabled();
+
+		if (bLowLatencyModeEnabled)
+		{
+			SetEnabled(false);
+
+			if (LocalPlayer && LocalPlayer->ViewportClient && LocalPlayer->ViewportClient->ViewportConsole)
+			{
+				LocalPlayer->ViewportClient->ViewportConsole->OutputText("Reflex Low Latency mode: Off");
+			}
+
+			UE_LOG(LogMaxTickRateHandler, Log, TEXT("Reflex Low Latency mode: Off"));
+		}
+		else
+		{
+			SetEnabled(true);
+			SetFlags(1);
+
+			if (LocalPlayer && LocalPlayer->ViewportClient && LocalPlayer->ViewportClient->ViewportConsole)
+			{
+				LocalPlayer->ViewportClient->ViewportConsole->OutputText("Reflex Low Latency mode: On");
+			}
+
+			UE_LOG(LogMaxTickRateHandler, Log, TEXT("Reflex Low Latency mode: On"));
+		}
+
+		bHandled = true;
+	}
+
+	return bHandled;
 }
\ No newline at end of file
diff --git a/Engine/Plugins/Runtime/Nvidia/Reflex/Source/Public/ReflexLatencyMarkers.h b/Engine/Plugins/Runtime/Nvidia/Reflex/Source/Public/ReflexLatencyMarkers.h
--- a/Engine/Plugins/Runtime/Nvidia/Reflex/Source/Public/ReflexLatencyMarkers.h
+++ b/Engine/Plugins/Runtime/Nvidia/Reflex/Source/Public/ReflexLatencyMarkers.h
@@ -6,26 +6,33 @@
 #include "Tickable.h"
 #include "Performance/LatencyMarkerModule.h"
 
-class FReflexLatencyMarkers : public ILatencyMarkerModule, public FTickableGameObject
+class FReflexLatencyMarkers : public ILatencyMarkerModule, public FTickableGameObject, public FSelfRegisteringExec
 {
 public:
 	virtual ~FReflexLatencyMarkers() {}
 
 	bool bProperDriverVersion = false;
+	bool bFeatureSupport = false;
 
 	float AverageTotalLatencyMs = 0.0f;
 	float AverageGameLatencyMs = 0.0f;
 	float AverageRenderLatencyMs = 0.0f;
+
+	float AverageSimulationLatencyMs = 0.0f;
+	float AverageRenderSubmitLatencyMs = 0.0f;
+	float AveragePresentLatencyMs = 0.0f;
 	float AverageDriverLatencyMs = 0.0f;
-	float AverageOSWorkQueueLatencyMs = 0.0f;
+	float AverageOSRenderQueueLatencyMs = 0.0f;
 	float AverageGPURenderLatencyMs = 0.0f;
 
-	float RenderOffsetMs = 0.0f;
+	float RenderSubmitOffsetMs = 0.0f;
+	float PresentOffsetMs = 0.0f;
 	float DriverOffsetMs = 0.0f;
-	float OSWorkQueueOffsetMs = 0.0f;
+	float OSRenderQueueOffsetMs = 0.0f;
 	float GPURenderOffsetMs = 0.0f;
 
 	bool bEnabled = false;
+	bool bFlashIndicatorEnabled = false;
 
 	virtual void Tick(float DeltaTime);
 	virtual bool IsTickable() const { return true; }
@@ -35,24 +42,38 @@ public:
 
 	virtual void Initialize() override;
 
-	virtual void SetEnabled(bool bInEnabled) override { bEnabled = bInEnabled; }
+	virtual void SetEnabled(bool bInEnabled) override;
 	virtual bool GetEnabled() override;
+	virtual bool GetAvailable() override;
+	virtual void SetFlashIndicatorEnabled(bool bInEnabled) override;
+	virtual bool GetFlashIndicatorEnabled() override;
+
+	virtual bool Exec(UWorld* InWorld, const TCHAR* Cmd, FOutputDevice& Ar) override;
 
-	virtual void SetGameLatencyMarkerStart(uint64 FrameNumber) override;
-	virtual void SetGameLatencyMarkerEnd(uint64 FrameNumber) override;
-	virtual void SetRenderLatencyMarkerStart(uint64 FrameNumber) override;
-	virtual void SetRenderLatencyMarkerEnd(uint64 FrameNumber) override;
+	virtual void SetInputSampleLatencyMarker(uint64 FrameNumber) override;
+	virtual void SetSimulationLatencyMarkerStart(uint64 FrameNumber) override;
+	virtual void SetSimulationLatencyMarkerEnd(uint64 FrameNumber) override;
+	virtual void SetRenderSubmitLatencyMarkerStart(uint64 FrameNumber) override;
+	virtual void SetRenderSubmitLatencyMarkerEnd(uint64 FrameNumber) override;
+	virtual void SetPresentLatencyMarkerStart(uint64 FrameNumber) override;
+	virtual void SetPresentLatencyMarkerEnd(uint64 FrameNumber) override;
+	virtual void SetFlashIndicatorLatencyMarker(uint64 FrameNumber) override;
 	virtual void SetCustomLatencyMarker(uint32 MarkerId, uint64 FrameNumber) override;
 
 	virtual float GetTotalLatencyInMs() override { return AverageTotalLatencyMs; }
-	virtual float GetGameLatencyInMs() override { return AverageGameLatencyMs; }
-	virtual float GetRenderLatencyInMs() override { return AverageRenderLatencyMs; }
+	virtual float GetGameLatencyInMs() override { return AverageGameLatencyMs; } // This is defined as "Game simulation start to driver submission end"
+	virtual float GetRenderLatencyInMs() override { return AverageRenderLatencyMs; } // This is defined as "OS render queue start to GPU render end"
+
+	virtual float GetSimulationLatencyInMs() override { return AverageSimulationLatencyMs; }
+	virtual float GetRenderSubmitLatencyInMs() override { return AverageRenderSubmitLatencyMs; }
+	virtual float GetPresentLatencyInMs() override { return AveragePresentLatencyMs; }
 	virtual float GetDriverLatencyInMs() override { return AverageDriverLatencyMs; }
-	virtual float GetOSWorkQueueLatencyInMs() override { return AverageOSWorkQueueLatencyMs; }
+	virtual float GetOSRenderQueueLatencyInMs() override { return AverageOSRenderQueueLatencyMs; }
 	virtual float GetGPURenderLatencyInMs() override { return AverageGPURenderLatencyMs; }
 
-	virtual float GetRenderOffsetFromGameInMs() override { return RenderOffsetMs; }
-	virtual float GetDriverOffsetFromGameInMs() override { return DriverOffsetMs; }
-	virtual float GetOSWorkQueueOffsetFromGameInMs() override { return OSWorkQueueOffsetMs; }
-	virtual float GetGPURenderOffsetFromGameInMs() override { return GPURenderOffsetMs; }
+	virtual float GetRenderSubmitOffsetFromFrameStartInMs() override { return RenderSubmitOffsetMs; }
+	virtual float GetPresentOffsetFromFrameStartInMs() override { return PresentOffsetMs; }
+	virtual float GetDriverOffsetFromFrameStartInMs() override { return DriverOffsetMs; }
+	virtual float GetOSRenderQueueOffsetFromFrameStartInMs() override { return OSRenderQueueOffsetMs; }
+	virtual float GetGPURenderOffsetFromFrameStartInMs() override { return GPURenderOffsetMs; }
 };
\ No newline at end of file
diff --git a/Engine/Plugins/Runtime/Nvidia/Reflex/Source/Public/ReflexMaxTickRateHandler.h b/Engine/Plugins/Runtime/Nvidia/Reflex/Source/Public/ReflexMaxTickRateHandler.h
--- a/Engine/Plugins/Runtime/Nvidia/Reflex/Source/Public/ReflexMaxTickRateHandler.h
+++ b/Engine/Plugins/Runtime/Nvidia/Reflex/Source/Public/ReflexMaxTickRateHandler.h
@@ -5,14 +5,17 @@
 #include "CoreMinimal.h"
 #include "Performance/MaxTickRateHandlerModule.h"
 
-class FReflexMaxTickRateHandler : public IMaxTickRateHandlerModule
+class FReflexMaxTickRateHandler : public IMaxTickRateHandlerModule, public FSelfRegisteringExec
 {
 public:
 	virtual ~FReflexMaxTickRateHandler() {}
 
 	virtual void Initialize() override;
-	virtual void SetEnabled(bool bInEnabled) override { bEnabled = bInEnabled; }
-	virtual bool GetEnabled() override { return bEnabled; }
+	virtual void SetEnabled(bool bInEnabled) override;
+	virtual bool GetEnabled() override;
+	virtual bool GetAvailable() override;
+
+	virtual bool Exec(UWorld* InWorld, const TCHAR* Cmd, FOutputDevice& Ar) override;
 
 	// Used to provide a generic customization interface for custom tick rate handlers
 	virtual void SetFlags(uint32 Flags);
@@ -23,9 +26,10 @@ public:
 	bool bEnabled = false;
 	bool bWasEnabled = false;
 	bool bProperDriverVersion = false;
+	bool bFeatureSupport = false;
 	float MinimumInterval = -1.0f;
-	bool bUltraLowLatency = true;
-	bool bGPUBoost = true;
+	bool bLowLatencyMode = true;
+	bool bBoost = false;
 
 	uint32 CustomFlags = 0;
 	uint32 LastCustomFlags = 0;
diff --git a/Engine/Plugins/Runtime/Nvidia/Reflex/Source/Reflex.Build.cs b/Engine/Plugins/Runtime/Nvidia/Reflex/Source/Reflex.Build.cs
--- a/Engine/Plugins/Runtime/Nvidia/Reflex/Source/Reflex.Build.cs
+++ b/Engine/Plugins/Runtime/Nvidia/Reflex/Source/Reflex.Build.cs
@@ -16,6 +16,7 @@ namespace UnrealBuildTool.Rules
 					"Core",
 					"Engine",
 					"RHI",
+					"CoreUObject"
 				}
 			);
 
diff --git a/Engine/Source/Runtime/Core/Private/Misc/CoreGlobals.cpp b/Engine/Source/Runtime/Core/Private/Misc/CoreGlobals.cpp
--- a/Engine/Source/Runtime/Core/Private/Misc/CoreGlobals.cpp
+++ b/Engine/Source/Runtime/Core/Private/Misc/CoreGlobals.cpp
@@ -347,6 +347,8 @@ bool					GEventDrivenLoaderEnabled = false;
 
 /** Steadily increasing frame counter.																		*/
 uint64					GFrameCounter					= 0;
+uint64					GFrameCounterRenderThread		= 0;
+
 uint64					GLastGCFrame					= 0;
 /** The time input was sampled, in cycles. */
 uint64					GInputTime					= 0;
diff --git a/Engine/Source/Runtime/Core/Public/CoreGlobals.h b/Engine/Source/Runtime/Core/Public/CoreGlobals.h
--- a/Engine/Source/Runtime/Core/Public/CoreGlobals.h
+++ b/Engine/Source/Runtime/Core/Public/CoreGlobals.h
@@ -438,6 +438,8 @@ extern CORE_API bool GIsRetrievingVTablePtr;
 /** Steadily increasing frame counter. */
 extern CORE_API uint64 GFrameCounter;
 
+extern CORE_API uint64 GFrameCounterRenderThread;
+
 /** GFrameCounter the last time GC was run. */
 extern CORE_API uint64 GLastGCFrame;
 
diff --git a/Engine/Source/Runtime/Engine/Classes/Engine/Engine.h b/Engine/Source/Runtime/Engine/Classes/Engine/Engine.h
--- a/Engine/Source/Runtime/Engine/Classes/Engine/Engine.h
+++ b/Engine/Source/Runtime/Engine/Classes/Engine/Engine.h
@@ -2093,11 +2093,18 @@ public:
 	 */
 	virtual void UpdateTimeAndHandleMaxTickRate();
 
-	void SetGameLatencyMarkerStart(uint64 FrameNumber);
-	void SetGameLatencyMarkerEnd(uint64 FrameNumber);
+	void SetInputSampleLatencyMarker(uint64 FrameNumber);
 
-	void SetRenderLatencyMarkerStart(uint64 FrameNumber);
-	void SetRenderLatencyMarkerEnd(uint64 FrameNumber);
+	void SetSimulationLatencyMarkerStart(uint64 FrameNumber);
+	void SetSimulationLatencyMarkerEnd(uint64 FrameNumber);
+
+	void SetPresentLatencyMarkerStart(uint64 FrameNumber);
+	void SetPresentLatencyMarkerEnd(uint64 FrameNumber);
+
+	void SetRenderSubmitLatencyMarkerStart(uint64 FrameNumber);
+	void SetRenderSubmitLatencyMarkerEnd(uint64 FrameNumber);
+
+	void SetFlashIndicatorLatencyMarker(uint64 FrameNumber);
 
 	/**
 	 * Allows games to correct the negative delta
diff --git a/Engine/Source/Runtime/Engine/Private/GameViewportClient.cpp b/Engine/Source/Runtime/Engine/Private/GameViewportClient.cpp
--- a/Engine/Source/Runtime/Engine/Private/GameViewportClient.cpp
+++ b/Engine/Source/Runtime/Engine/Private/GameViewportClient.cpp
@@ -602,6 +602,11 @@ bool UGameViewportClient::InputKey(const FInputKeyEventArgs& EventArgs)
 		return true;
 	}
 
+	if (EventArgs.Key == EKeys::LeftMouseButton && EventArgs.Event == EInputEvent::IE_Pressed)
+	{
+		GEngine->SetFlashIndicatorLatencyMarker(GFrameCounter);
+	}
+
 	if (IgnoreInput())
 	{
 		return ViewportConsole ? ViewportConsole->InputKey(ControllerId, EventArgs.Key, EventArgs.Event, EventArgs.AmountDepressed, EventArgs.IsGamepad()) : false;
diff --git a/Engine/Source/Runtime/Engine/Private/UnrealClient.cpp b/Engine/Source/Runtime/Engine/Private/UnrealClient.cpp
--- a/Engine/Source/Runtime/Engine/Private/UnrealClient.cpp
+++ b/Engine/Source/Runtime/Engine/Private/UnrealClient.cpp
@@ -1392,10 +1392,20 @@ void FViewport::EndRenderFrame(FRHICommandListImmediate& RHICmdList, bool bPrese
 {
 	check( IsInRenderingThread() );
 
+	RHICmdList.EnqueueLambda([CurrentFrameCounter = GFrameCounterRenderThread](FRHICommandListImmediate& InRHICmdList)
+	{
+		GEngine->SetPresentLatencyMarkerStart(CurrentFrameCounter);
+	});
+
 	uint32 StartTime = FPlatformTime::Cycles();
 	RHICmdList.EndDrawingViewport(GetViewportRHI(), bPresent, bLockToVsync);
 	uint32 EndTime = FPlatformTime::Cycles();
 
+	RHICmdList.EnqueueLambda([CurrentFrameCounter = GFrameCounterRenderThread](FRHICommandListImmediate& InRHICmdList)
+	{
+		GEngine->SetPresentLatencyMarkerEnd(CurrentFrameCounter);
+	});
+
 	GRenderThreadIdle[ERenderThreadIdleTypes::WaitingForGPUPresent] += EndTime - StartTime;
 	GRenderThreadNumIdle[ERenderThreadIdleTypes::WaitingForGPUPresent]++;
 }
diff --git a/Engine/Source/Runtime/Engine/Private/UnrealEngine.cpp b/Engine/Source/Runtime/Engine/Private/UnrealEngine.cpp
--- a/Engine/Source/Runtime/Engine/Private/UnrealEngine.cpp
+++ b/Engine/Source/Runtime/Engine/Private/UnrealEngine.cpp
@@ -2091,43 +2091,84 @@ double UEngine::CorrectNegativeTimeDelta(double DeltaRealTime)
 	return 0.01;
 }
 
-void UEngine::SetGameLatencyMarkerStart(uint64 FrameNumber)
+void UEngine::SetInputSampleLatencyMarker(uint64 FrameNumber)
 {
 	TArray<ILatencyMarkerModule*> LatencyMarkerModules = IModularFeatures::Get().GetModularFeatureImplementations<ILatencyMarkerModule>(ILatencyMarkerModule::GetModularFeatureName());
 
 	for (ILatencyMarkerModule* LatencyMarkerModule : LatencyMarkerModules)
 	{
-		LatencyMarkerModule->SetGameLatencyMarkerStart(FrameNumber);
+		LatencyMarkerModule->SetInputSampleLatencyMarker(FrameNumber);
 	}
 }
 
-void UEngine::SetGameLatencyMarkerEnd(uint64 FrameNumber)
+void UEngine::SetSimulationLatencyMarkerStart(uint64 FrameNumber)
 {
 	TArray<ILatencyMarkerModule*> LatencyMarkerModules = IModularFeatures::Get().GetModularFeatureImplementations<ILatencyMarkerModule>(ILatencyMarkerModule::GetModularFeatureName());
 
 	for (ILatencyMarkerModule* LatencyMarkerModule : LatencyMarkerModules)
 	{
-		LatencyMarkerModule->SetGameLatencyMarkerEnd(FrameNumber);
+		LatencyMarkerModule->SetSimulationLatencyMarkerStart(FrameNumber);
 	}
 }
 
-void UEngine::SetRenderLatencyMarkerStart(uint64 FrameNumber)
+void UEngine::SetSimulationLatencyMarkerEnd(uint64 FrameNumber)
 {
 	TArray<ILatencyMarkerModule*> LatencyMarkerModules = IModularFeatures::Get().GetModularFeatureImplementations<ILatencyMarkerModule>(ILatencyMarkerModule::GetModularFeatureName());
 
 	for (ILatencyMarkerModule* LatencyMarkerModule : LatencyMarkerModules)
 	{
-		LatencyMarkerModule->SetRenderLatencyMarkerStart(FrameNumber);
+		LatencyMarkerModule->SetSimulationLatencyMarkerEnd(FrameNumber);
 	}
 }
 
-void UEngine::SetRenderLatencyMarkerEnd(uint64 FrameNumber)
+void UEngine::SetPresentLatencyMarkerStart(uint64 FrameNumber)
 {
 	TArray<ILatencyMarkerModule*> LatencyMarkerModules = IModularFeatures::Get().GetModularFeatureImplementations<ILatencyMarkerModule>(ILatencyMarkerModule::GetModularFeatureName());
 
 	for (ILatencyMarkerModule* LatencyMarkerModule : LatencyMarkerModules)
 	{
-		LatencyMarkerModule->SetRenderLatencyMarkerEnd(FrameNumber);
+		LatencyMarkerModule->SetPresentLatencyMarkerStart(FrameNumber);
+	}
+}
+
+void UEngine::SetPresentLatencyMarkerEnd(uint64 FrameNumber)
+{
+	TArray<ILatencyMarkerModule*> LatencyMarkerModules = IModularFeatures::Get().GetModularFeatureImplementations<ILatencyMarkerModule>(ILatencyMarkerModule::GetModularFeatureName());
+
+	for (ILatencyMarkerModule* LatencyMarkerModule : LatencyMarkerModules)
+	{
+		LatencyMarkerModule->SetPresentLatencyMarkerEnd(FrameNumber);
+	}
+}
+
+
+void UEngine::SetRenderSubmitLatencyMarkerStart(uint64 FrameNumber)
+{
+	TArray<ILatencyMarkerModule*> LatencyMarkerModules = IModularFeatures::Get().GetModularFeatureImplementations<ILatencyMarkerModule>(ILatencyMarkerModule::GetModularFeatureName());
+
+	for (ILatencyMarkerModule* LatencyMarkerModule : LatencyMarkerModules)
+	{
+		LatencyMarkerModule->SetRenderSubmitLatencyMarkerStart(FrameNumber);
+	}
+}
+
+void UEngine::SetRenderSubmitLatencyMarkerEnd(uint64 FrameNumber)
+{
+	TArray<ILatencyMarkerModule*> LatencyMarkerModules = IModularFeatures::Get().GetModularFeatureImplementations<ILatencyMarkerModule>(ILatencyMarkerModule::GetModularFeatureName());
+
+	for (ILatencyMarkerModule* LatencyMarkerModule : LatencyMarkerModules)
+	{
+		LatencyMarkerModule->SetRenderSubmitLatencyMarkerEnd(FrameNumber);
+	}
+}
+
+void UEngine::SetFlashIndicatorLatencyMarker(uint64 FrameNumber)
+{
+	TArray<ILatencyMarkerModule*> LatencyMarkerModules = IModularFeatures::Get().GetModularFeatureImplementations<ILatencyMarkerModule>(ILatencyMarkerModule::GetModularFeatureName());
+
+	for (ILatencyMarkerModule* LatencyMarkerModule : LatencyMarkerModules)
+	{
+		LatencyMarkerModule->SetFlashIndicatorLatencyMarker(FrameNumber);
 	}
 }
 
diff --git a/Engine/Source/Runtime/Engine/Public/Performance/LatencyMarkerModule.h b/Engine/Source/Runtime/Engine/Public/Performance/LatencyMarkerModule.h
--- a/Engine/Source/Runtime/Engine/Public/Performance/LatencyMarkerModule.h
+++ b/Engine/Source/Runtime/Engine/Public/Performance/LatencyMarkerModule.h
@@ -18,22 +18,34 @@ public:
 
 	virtual void SetEnabled(bool bInEnabled) = 0;
 	virtual bool GetEnabled() = 0;
-	virtual void SetGameLatencyMarkerStart(uint64 FrameNumber) = 0;
-	virtual void SetGameLatencyMarkerEnd(uint64 FrameNumber) = 0;
-	virtual void SetRenderLatencyMarkerStart(uint64 FrameNumber) = 0;
-	virtual void SetRenderLatencyMarkerEnd(uint64 FrameNumber) = 0;
-	
+	virtual bool GetAvailable() = 0;
+	virtual void SetFlashIndicatorEnabled(bool bInEnabled) = 0;
+	virtual bool GetFlashIndicatorEnabled() = 0;
+
+	virtual void SetInputSampleLatencyMarker(uint64 FrameNumber) = 0;
+	virtual void SetSimulationLatencyMarkerStart(uint64 FrameNumber) = 0;
+	virtual void SetSimulationLatencyMarkerEnd(uint64 FrameNumber) = 0;
+	virtual void SetPresentLatencyMarkerStart(uint64 FrameNumber) = 0;
+	virtual void SetPresentLatencyMarkerEnd(uint64 FrameNumber) = 0;
+	virtual void SetRenderSubmitLatencyMarkerStart(uint64 FrameNumber) = 0;
+	virtual void SetRenderSubmitLatencyMarkerEnd(uint64 FrameNumber) = 0;
+	virtual void SetFlashIndicatorLatencyMarker(uint64 FrameNumber) = 0;
 	virtual void SetCustomLatencyMarker(uint32 MarkerId, uint64 FrameNumber) = 0;
 
 	virtual float GetTotalLatencyInMs() = 0;
-	virtual float GetGameLatencyInMs() = 0;
-	virtual float GetRenderLatencyInMs() = 0;
+	virtual float GetGameLatencyInMs() = 0; // This is defined as "Game simulation start to driver submission end"
+	virtual float GetRenderLatencyInMs() = 0; // This is defined as "OS render queue start to GPU render end"
+
+	virtual float GetSimulationLatencyInMs() = 0;
+	virtual float GetRenderSubmitLatencyInMs() = 0;
+	virtual float GetPresentLatencyInMs() = 0;
 	virtual float GetDriverLatencyInMs() = 0;
-	virtual float GetOSWorkQueueLatencyInMs() = 0;
+	virtual float GetOSRenderQueueLatencyInMs() = 0;
 	virtual float GetGPURenderLatencyInMs() = 0;
 
-	virtual float GetRenderOffsetFromGameInMs() = 0;
-	virtual float GetDriverOffsetFromGameInMs() = 0;
-	virtual float GetOSWorkQueueOffsetFromGameInMs() = 0;
-	virtual float GetGPURenderOffsetFromGameInMs() = 0;
+	virtual float GetRenderSubmitOffsetFromFrameStartInMs() = 0;
+	virtual float GetPresentOffsetFromFrameStartInMs() = 0;
+	virtual float GetDriverOffsetFromFrameStartInMs() = 0;
+	virtual float GetOSRenderQueueOffsetFromFrameStartInMs() = 0;
+	virtual float GetGPURenderOffsetFromFrameStartInMs() = 0;
 };
diff --git a/Engine/Source/Runtime/Engine/Public/Performance/MaxTickRateHandlerModule.h b/Engine/Source/Runtime/Engine/Public/Performance/MaxTickRateHandlerModule.h
--- a/Engine/Source/Runtime/Engine/Public/Performance/MaxTickRateHandlerModule.h
+++ b/Engine/Source/Runtime/Engine/Public/Performance/MaxTickRateHandlerModule.h
@@ -17,6 +17,7 @@ public:
 	virtual void Initialize() = 0;
 	virtual void SetEnabled(bool bInEnabled) = 0;
 	virtual bool GetEnabled() = 0;
+	virtual bool GetAvailable() = 0;
 
 	virtual void SetFlags(uint32 Flags) = 0;
 	virtual uint32 GetFlags() = 0;
diff --git a/Engine/Source/Runtime/Launch/Private/LaunchEngineLoop.cpp b/Engine/Source/Runtime/Launch/Private/LaunchEngineLoop.cpp
--- a/Engine/Source/Runtime/Launch/Private/LaunchEngineLoop.cpp
+++ b/Engine/Source/Runtime/Launch/Private/LaunchEngineLoop.cpp
@@ -4814,7 +4814,7 @@ static inline void BeginFrameRenderThread(FRHICommandListImmediate& RHICmdList,
 
 	RHICmdList.EnqueueLambda([CurrentFrameCounter](FRHICommandListImmediate& InRHICmdList)
 	{
-		GEngine->SetRenderLatencyMarkerStart(CurrentFrameCounter);
+		GEngine->SetRenderSubmitLatencyMarkerStart(CurrentFrameCounter);
 	});
 }
 
@@ -4823,7 +4823,7 @@ static inline void EndFrameRenderThread(FRHICommandListImmediate& RHICmdList, ui
 {
 	RHICmdList.EnqueueLambda([CurrentFrameCounter](FRHICommandListImmediate& InRHICmdList)
 	{
-		GEngine->SetRenderLatencyMarkerEnd(CurrentFrameCounter);
+		GEngine->SetRenderSubmitLatencyMarkerEnd(CurrentFrameCounter);
 	});
 
 	FCoreDelegates::OnEndFrameRT.Broadcast();
@@ -4956,7 +4956,7 @@ void FEngineLoop::Tick()
 		{
 			QUICK_SCOPE_CYCLE_COUNTER(STAT_FEngineLoop_UpdateTimeAndHandleMaxTickRate);
 			GEngine->UpdateTimeAndHandleMaxTickRate();
-			GEngine->SetGameLatencyMarkerStart(CurrentFrameCounter);
+			GEngine->SetSimulationLatencyMarkerStart(CurrentFrameCounter);
 		}
 
 		for (const FWorldContext& Context : GEngine->GetWorldContexts())
@@ -5042,6 +5042,8 @@ void FEngineLoop::Tick()
 		// will pass us messages instead.
 		if (!GUELibraryOverrideSettings.bIsEmbedded)
 		{
+			GEngine->SetInputSampleLatencyMarker(CurrentFrameCounter);
+
 			//QUICK_SCOPE_CYCLE_COUNTER(STAT_PumpMessages);
 			FPlatformApplicationMisc::PumpMessages(true);
 		}
@@ -5288,9 +5290,19 @@ void FEngineLoop::Tick()
 			RHITick( FApp::GetDeltaTime() ); // Update RHI.
 		}
 
+		// We need to set this marker before EndFrameRenderThread is enqueued. 
+		// If multithreaded rendering is off, it can cause a bad ordering of game and rendering markers.
+		GEngine->SetSimulationLatencyMarkerEnd(CurrentFrameCounter);
+
 		// Increment global frame counter. Once for each engine tick.
 		GFrameCounter++;
 
+		ENQUEUE_RENDER_COMMAND(FrameCounter)(
+			[CurrentFrameCounter = GFrameCounter](FRHICommandListImmediate& RHICmdList)
+		{
+			GFrameCounterRenderThread = CurrentFrameCounter;
+		});
+
 		// Disregard first few ticks for total tick time as it includes loading and such.
 		if (GFrameCounter > 6)
 		{
@@ -5346,10 +5358,6 @@ void FEngineLoop::Tick()
 		}
 		#endif
 
-		// We need to set this marker before EndFrameRenderThread is enqueued. 
-		// If multithreaded rendering is off, it can cause a bad ordering of game and rendering markers.
-		GEngine->SetGameLatencyMarkerEnd(CurrentFrameCounter);
-
 		// end of RHI frame
 		ENQUEUE_RENDER_COMMAND(EndFrame)(
 			[CurrentFrameCounter](FRHICommandListImmediate& RHICmdList)
diff --git a/Engine/Source/Runtime/SlateRHIRenderer/Private/SlateRHIRenderer.cpp b/Engine/Source/Runtime/SlateRHIRenderer/Private/SlateRHIRenderer.cpp
--- a/Engine/Source/Runtime/SlateRHIRenderer/Private/SlateRHIRenderer.cpp
+++ b/Engine/Source/Runtime/SlateRHIRenderer/Private/SlateRHIRenderer.cpp
@@ -1138,8 +1138,18 @@ void FSlateRHIRenderer::DrawWindow_RenderThread(FRHICommandListImmediate& RHICmd
 	// Calculate renderthread time (excluding idle time).	
 	uint32 StartTime = FPlatformTime::Cycles();
 
+	RHICmdList.EnqueueLambda([CurrentFrameCounter = GFrameCounterRenderThread](FRHICommandListImmediate& InRHICmdList)
+	{
+		GEngine->SetPresentLatencyMarkerStart(CurrentFrameCounter);
+	});
+
 	RHICmdList.EndDrawingViewport(ViewportInfo.ViewportRHI, true, DrawCommandParams.bLockToVsync);
 
+	RHICmdList.EnqueueLambda([CurrentFrameCounter = GFrameCounterRenderThread](FRHICommandListImmediate& InRHICmdList)
+	{
+		GEngine->SetPresentLatencyMarkerEnd(CurrentFrameCounter);
+	});
+
 	uint32 EndTime = FPlatformTime::Cycles();
 
 	GSwapBufferTime = EndTime - StartTime;
