Reflex Blueprint library (on behalf of Nvidia)

Adds blueprint functions to Reflex plugin to export reflex plugin data and capabilities for use in blueprints.

#rnx

diff --git a/Engine/Plugins/Runtime/Nvidia/Reflex/Source/Private/ReflexBlueprint.cpp b/Engine/Plugins/Runtime/Nvidia/Reflex/Source/Private/ReflexBlueprint.cpp
new file mode 100644
--- /dev/null
+++ b/Engine/Plugins/Runtime/Nvidia/Reflex/Source/Private/ReflexBlueprint.cpp
@@ -0,0 +1,140 @@
+﻿/*
+* Copyright (c) 2021s NVIDIA CORPORATION.  All rights reserved.
+*
+* NVIDIA Corporation and its licensors retain all intellectual property and proprietary
+* rights in and to this software, related documentation and any modifications thereto.
+* Any use, reproduction, disclosure or distribution of this software and related
+* documentation without an express license agreement from NVIDIA Corporation is strictly
+* prohibited.
+*
+* TO THE MAXIMUM EXTENT PERMITTED BY APPLICABLE LAW, THIS SOFTWARE IS PROVIDED *AS IS*
+* AND NVIDIA AND ITS SUPPLIERS DISCLAIM ALL WARRANTIES, EITHER EXPRESS OR IMPLIED,
+* INCLUDING, BUT NOT LIMITED TO, IMPLIED WARRANTIES OF MERCHANTABILITY AND FITNESS FOR A
+* PARTICULAR PURPOSE.  IN NO EVENT SHALL NVIDIA OR ITS SUPPLIERS BE LIABLE FOR ANY
+* SPECIAL, INCIDENTAL, INDIRECT, OR CONSEQUENTIAL DAMAGES WHATSOEVER (INCLUDING, WITHOUT
+* LIMITATION, DAMAGES FOR LOSS OF BUSINESS PROFITS, BUSINESS INTERRUPTION, LOSS OF
+* BUSINESS INFORMATION, OR ANY OTHER PECUNIARY LOSS) ARISING OUT OF THE USE OF OR
+* INABILITY TO USE THIS SOFTWARE, EVEN IF NVIDIA HAS BEEN ADVISED OF THE POSSIBILITY OF
+* SUCH DAMAGES.
+*/
+#include "ReflexBlueprint.h"
+
+#include "ReflexLatencyMarkers.h"
+#include "ReflexMaxTickRateHandler.h"	
+
+UReflexBlueprintLibrary::UReflexBlueprintLibrary(const FObjectInitializer& ObjectInitializer)
+	: Super(ObjectInitializer)
+{
+}
+
+bool UReflexBlueprintLibrary::GetReflexAvailable()
+{
+	bool bIsMaxTickRateHandlerEnabled = false;
+
+	TArray<IMaxTickRateHandlerModule*> MaxTickRateHandlerModules = IModularFeatures::Get()
+		.GetModularFeatureImplementations<IMaxTickRateHandlerModule>(IMaxTickRateHandlerModule::GetModularFeatureName());
+	for (IMaxTickRateHandlerModule* MaxTickRateHandler : MaxTickRateHandlerModules)
+	{
+		bIsMaxTickRateHandlerEnabled = bIsMaxTickRateHandlerEnabled || MaxTickRateHandler->GetAvailable();
+	}
+
+	bool bIsLatencyMarkerModuleEnabled = false;
+
+	TArray<ILatencyMarkerModule*> LatencyMarkerModules = IModularFeatures::Get()
+		.GetModularFeatureImplementations<ILatencyMarkerModule>(ILatencyMarkerModule::GetModularFeatureName());
+	for (ILatencyMarkerModule* LatencyMarkerModule : LatencyMarkerModules)
+	{
+		bIsLatencyMarkerModuleEnabled = bIsLatencyMarkerModuleEnabled || LatencyMarkerModule->GetAvailable(); 
+	}
+	
+	return bIsMaxTickRateHandlerEnabled && bIsLatencyMarkerModuleEnabled;
+
+}
+
+void UReflexBlueprintLibrary::SetReflexMode(const EReflexMode Mode)
+{
+	TArray<IMaxTickRateHandlerModule*> MaxTickRateHandlerModules = IModularFeatures::Get()
+		.GetModularFeatureImplementations<IMaxTickRateHandlerModule>(IMaxTickRateHandlerModule::GetModularFeatureName());
+	for (IMaxTickRateHandlerModule* MaxTickRateHandler : MaxTickRateHandlerModules)
+	{
+		MaxTickRateHandler->SetEnabled(Mode!=EReflexMode::Disabled);
+		MaxTickRateHandler->SetFlags(static_cast<uint32>(Mode));
+	}
+
+	TArray<ILatencyMarkerModule*> LatencyMarkerModules = IModularFeatures::Get()
+		.GetModularFeatureImplementations<ILatencyMarkerModule>(ILatencyMarkerModule::GetModularFeatureName());
+	for (ILatencyMarkerModule* LatencyMarkerModule : LatencyMarkerModules)
+	{
+		LatencyMarkerModule->SetEnabled(true);
+	}
+}
+
+EReflexMode UReflexBlueprintLibrary::GetReflexMode()
+{
+	TArray<IMaxTickRateHandlerModule*> MaxTickRateHandlerModules = IModularFeatures::Get()
+		.GetModularFeatureImplementations<IMaxTickRateHandlerModule>(IMaxTickRateHandlerModule::GetModularFeatureName());
+	for (IMaxTickRateHandlerModule* MaxTickRateHandler : MaxTickRateHandlerModules)
+	{
+		return static_cast<EReflexMode>(MaxTickRateHandler->GetFlags());
+	}
+
+	return EReflexMode::Disabled;
+}
+
+void UReflexBlueprintLibrary::SetFlashIndicatorEnabled(const bool bEnabled)
+{
+	TArray<ILatencyMarkerModule*> LatencyMarkerModules = IModularFeatures::Get()
+		.GetModularFeatureImplementations<ILatencyMarkerModule>(ILatencyMarkerModule::GetModularFeatureName());
+	for (ILatencyMarkerModule* LatencyMarkerModule : LatencyMarkerModules)
+	{
+		LatencyMarkerModule->SetFlashIndicatorEnabled(bEnabled);
+	}
+}
+
+bool UReflexBlueprintLibrary::GetFlashIndicatorEnabled()
+{
+	TArray<ILatencyMarkerModule*> LatencyMarkerModules = IModularFeatures::Get()
+		.GetModularFeatureImplementations<ILatencyMarkerModule>(ILatencyMarkerModule::GetModularFeatureName());
+	for (ILatencyMarkerModule* LatencyMarkerModule : LatencyMarkerModules)
+	{
+		return LatencyMarkerModule->GetFlashIndicatorEnabled();
+	}
+	
+	return false;
+}
+
+float UReflexBlueprintLibrary::GetGameToRenderLatencyInMs()
+{
+	TArray<ILatencyMarkerModule*> LatencyMarkerModules = IModularFeatures::Get()
+		.GetModularFeatureImplementations<ILatencyMarkerModule>(ILatencyMarkerModule::GetModularFeatureName());
+	for (ILatencyMarkerModule* LatencyMarkerModule : LatencyMarkerModules)
+	{
+		return LatencyMarkerModule->GetTotalLatencyInMs();
+	}
+
+	return 0.f;
+}
+
+float UReflexBlueprintLibrary::GetGameLatencyInMs()
+{
+	TArray<ILatencyMarkerModule*> LatencyMarkerModules = IModularFeatures::Get()
+		.GetModularFeatureImplementations<ILatencyMarkerModule>(ILatencyMarkerModule::GetModularFeatureName());
+	for (ILatencyMarkerModule* LatencyMarkerModule : LatencyMarkerModules)
+	{
+		return LatencyMarkerModule->GetGameLatencyInMs();
+	}
+
+	return 0.f;
+}
+
+float UReflexBlueprintLibrary::GetRenderLatencyInMs()
+{
+	TArray<ILatencyMarkerModule*> LatencyMarkerModules = IModularFeatures::Get()
+		.GetModularFeatureImplementations<ILatencyMarkerModule>(ILatencyMarkerModule::GetModularFeatureName());
+	for (ILatencyMarkerModule* LatencyMarkerModule : LatencyMarkerModules)
+	{
+		return LatencyMarkerModule->GetRenderLatencyInMs();
+	}
+
+	return 0.f;
+}
diff --git a/Engine/Plugins/Runtime/Nvidia/Reflex/Source/Public/ReflexBlueprint.h b/Engine/Plugins/Runtime/Nvidia/Reflex/Source/Public/ReflexBlueprint.h
new file mode 100644
--- /dev/null
+++ b/Engine/Plugins/Runtime/Nvidia/Reflex/Source/Public/ReflexBlueprint.h
@@ -0,0 +1,59 @@
+﻿/*
+* Copyright (c) 2021 NVIDIA CORPORATION.  All rights reserved.
+*
+* NVIDIA Corporation and its licensors retain all intellectual property and proprietary
+* rights in and to this software, related documentation and any modifications thereto.
+* Any use, reproduction, disclosure or distribution of this software and related
+* documentation without an express license agreement from NVIDIA Corporation is strictly
+* prohibited.
+*
+* TO THE MAXIMUM EXTENT PERMITTED BY APPLICABLE LAW, THIS SOFTWARE IS PROVIDED *AS IS*
+* AND NVIDIA AND ITS SUPPLIERS DISCLAIM ALL WARRANTIES, EITHER EXPRESS OR IMPLIED,
+* INCLUDING, BUT NOT LIMITED TO, IMPLIED WARRANTIES OF MERCHANTABILITY AND FITNESS FOR A
+* PARTICULAR PURPOSE.  IN NO EVENT SHALL NVIDIA OR ITS SUPPLIERS BE LIABLE FOR ANY
+* SPECIAL, INCIDENTAL, INDIRECT, OR CONSEQUENTIAL DAMAGES WHATSOEVER (INCLUDING, WITHOUT
+* LIMITATION, DAMAGES FOR LOSS OF BUSINESS PROFITS, BUSINESS INTERRUPTION, LOSS OF
+* BUSINESS INFORMATION, OR ANY OTHER PECUNIARY LOSS) ARISING OUT OF THE USE OF OR
+* INABILITY TO USE THIS SOFTWARE, EVEN IF NVIDIA HAS BEEN ADVISED OF THE POSSIBILITY OF
+* SUCH DAMAGES.
+*/
+#pragma once
+
+#include "Kismet/BlueprintFunctionLibrary.h"
+
+#include "ReflexBlueprint.generated.h"
+
+UENUM(BlueprintType)
+enum class EReflexMode : uint8
+{
+	Disabled = 0 UMETA(DisplayName="Disabled"),
+	Enabled = 1 UMETA(DisplayName="Enabled"),
+	EnabledPlusBoost = 3 UMETA(DisplayName="Enabled + Boost")
+};
+
+
+UCLASS()
+class UReflexBlueprintLibrary : public UBlueprintFunctionLibrary
+{
+	GENERATED_UCLASS_BODY()
+
+	UFUNCTION(BlueprintPure, Category="Reflex")
+	static bool GetReflexAvailable();
+
+	UFUNCTION(BlueprintCallable, Category="Reflex")
+	static void SetReflexMode(const EReflexMode Mode);
+	UFUNCTION(BlueprintPure, Category="Reflex")
+	static EReflexMode GetReflexMode();
+
+	UFUNCTION(BlueprintCallable, Category="Reflex")
+	static void SetFlashIndicatorEnabled(const bool bEnabled);
+	UFUNCTION(BlueprintPure, Category="Reflex")
+	static bool GetFlashIndicatorEnabled();
+	
+	UFUNCTION(BlueprintPure, Category="Reflex")
+	static float GetGameToRenderLatencyInMs();
+	UFUNCTION(BlueprintPure, Category="Reflex")
+	static float GetGameLatencyInMs();
+	UFUNCTION(BlueprintPure, Category="Reflex")
+	static float GetRenderLatencyInMs();
+};
