Change copyright notice on newly added files.

#rnx

diff --git a/Engine/Plugins/Runtime/Nvidia/Reflex/Source/Private/ReflexBlueprint.cpp b/Engine/Plugins/Runtime/Nvidia/Reflex/Source/Private/ReflexBlueprint.cpp
--- a/Engine/Plugins/Runtime/Nvidia/Reflex/Source/Private/ReflexBlueprint.cpp
+++ b/Engine/Plugins/Runtime/Nvidia/Reflex/Source/Private/ReflexBlueprint.cpp
@@ -1,22 +1,5 @@
-﻿/*
-* Copyright (c) 2021s NVIDIA CORPORATION.  All rights reserved.
-*
-* NVIDIA Corporation and its licensors retain all intellectual property and proprietary
-* rights in and to this software, related documentation and any modifications thereto.
-* Any use, reproduction, disclosure or distribution of this software and related
-* documentation without an express license agreement from NVIDIA Corporation is strictly
-* prohibited.
-*
-* TO THE MAXIMUM EXTENT PERMITTED BY APPLICABLE LAW, THIS SOFTWARE IS PROVIDED *AS IS*
-* AND NVIDIA AND ITS SUPPLIERS DISCLAIM ALL WARRANTIES, EITHER EXPRESS OR IMPLIED,
-* INCLUDING, BUT NOT LIMITED TO, IMPLIED WARRANTIES OF MERCHANTABILITY AND FITNESS FOR A
-* PARTICULAR PURPOSE.  IN NO EVENT SHALL NVIDIA OR ITS SUPPLIERS BE LIABLE FOR ANY
-* SPECIAL, INCIDENTAL, INDIRECT, OR CONSEQUENTIAL DAMAGES WHATSOEVER (INCLUDING, WITHOUT
-* LIMITATION, DAMAGES FOR LOSS OF BUSINESS PROFITS, BUSINESS INTERRUPTION, LOSS OF
-* BUSINESS INFORMATION, OR ANY OTHER PECUNIARY LOSS) ARISING OUT OF THE USE OF OR
-* INABILITY TO USE THIS SOFTWARE, EVEN IF NVIDIA HAS BEEN ADVISED OF THE POSSIBILITY OF
-* SUCH DAMAGES.
-*/
+﻿// Copyright Epic Games, Inc. All Rights Reserved.
+
 #include "ReflexBlueprint.h"
 
 #include "ReflexLatencyMarkers.h"
diff --git a/Engine/Plugins/Runtime/Nvidia/Reflex/Source/Public/ReflexBlueprint.h b/Engine/Plugins/Runtime/Nvidia/Reflex/Source/Public/ReflexBlueprint.h
--- a/Engine/Plugins/Runtime/Nvidia/Reflex/Source/Public/ReflexBlueprint.h
+++ b/Engine/Plugins/Runtime/Nvidia/Reflex/Source/Public/ReflexBlueprint.h
@@ -1,22 +1,5 @@
-﻿/*
-* Copyright (c) 2021 NVIDIA CORPORATION.  All rights reserved.
-*
-* NVIDIA Corporation and its licensors retain all intellectual property and proprietary
-* rights in and to this software, related documentation and any modifications thereto.
-* Any use, reproduction, disclosure or distribution of this software and related
-* documentation without an express license agreement from NVIDIA Corporation is strictly
-* prohibited.
-*
-* TO THE MAXIMUM EXTENT PERMITTED BY APPLICABLE LAW, THIS SOFTWARE IS PROVIDED *AS IS*
-* AND NVIDIA AND ITS SUPPLIERS DISCLAIM ALL WARRANTIES, EITHER EXPRESS OR IMPLIED,
-* INCLUDING, BUT NOT LIMITED TO, IMPLIED WARRANTIES OF MERCHANTABILITY AND FITNESS FOR A
-* PARTICULAR PURPOSE.  IN NO EVENT SHALL NVIDIA OR ITS SUPPLIERS BE LIABLE FOR ANY
-* SPECIAL, INCIDENTAL, INDIRECT, OR CONSEQUENTIAL DAMAGES WHATSOEVER (INCLUDING, WITHOUT
-* LIMITATION, DAMAGES FOR LOSS OF BUSINESS PROFITS, BUSINESS INTERRUPTION, LOSS OF
-* BUSINESS INFORMATION, OR ANY OTHER PECUNIARY LOSS) ARISING OUT OF THE USE OF OR
-* INABILITY TO USE THIS SOFTWARE, EVEN IF NVIDIA HAS BEEN ADVISED OF THE POSSIBILITY OF
-* SUCH DAMAGES.
-*/
+﻿// Copyright Epic Games, Inc. All Rights Reserved.
+
 #pragma once
 
 #include "Kismet/BlueprintFunctionLibrary.h"
