Compilation fixes

diff --git a/Engine/Plugins/Runtime/Nvidia/Reflex/Source/Private/ReflexStatsLatencyMarkers.cpp b/Engine/Plugins/Runtime/Nvidia/Reflex/Source/Private/ReflexStatsLatencyMarkers.cpp
--- a/Engine/Plugins/Runtime/Nvidia/Reflex/Source/Private/ReflexStatsLatencyMarkers.cpp
+++ b/Engine/Plugins/Runtime/Nvidia/Reflex/Source/Private/ReflexStatsLatencyMarkers.cpp
@@ -9,8 +9,6 @@ __pragma(warning(disable:4191)) // hide a warning from TraceLoggingProvider.h
 #include "Windows/HideWindowsPlatformTypes.h"
 THIRD_PARTY_INCLUDES_END
 
-DEFINE_LOG_CATEGORY_STATIC(LogLatencyMarkers, Log, All);
-
 NVSTATS_DEFINE();
 
 FReflexStatsLatencyMarkers::~FReflexStatsLatencyMarkers()
@@ -113,7 +111,6 @@ bool FReflexStatsLatencyMarkers::ProcessMessage(HWND hwnd, uint32 msg, WPARAM wP
 
 	if (NVSTATS_IS_PING_MSG_ID(msg) || ((msg) == WM_KEYDOWN && ((wParam) == g_ReflexStatsVirtualKey)))
 	{
-		UE_LOG(LogLatencyMarkers, VeryVerbose, TEXT("Catch a ping event"));
 		SetCustomLatencyMarker(NVSTATS_PC_LATENCY_PING, GFrameCounter);
 		bHandled = true;
 	}
diff --git a/Engine/Source/Runtime/Core/Private/Features/ModularFeatures.cpp b/Engine/Source/Runtime/Core/Private/Features/ModularFeatures.cpp
--- a/Engine/Source/Runtime/Core/Private/Features/ModularFeatures.cpp
+++ b/Engine/Source/Runtime/Core/Private/Features/ModularFeatures.cpp
@@ -1,7 +1,7 @@
 // Copyright Epic Games, Inc. All Rights Reserved.
 
 #include "Features/ModularFeatures.h"
-
+#include "Misc/ScopeLock.h"
 
 IModularFeatures& IModularFeatures::Get()
 {
