Fix compile

diff --git a/Engine/Plugins/Runtime/Nvidia/Reflex/Source/Private/ReflexStatsLatencyMarkers.cpp b/Engine/Plugins/Runtime/Nvidia/Reflex/Source/Private/ReflexStatsLatencyMarkers.cpp
--- a/Engine/Plugins/Runtime/Nvidia/Reflex/Source/Private/ReflexStatsLatencyMarkers.cpp
+++ b/Engine/Plugins/Runtime/Nvidia/Reflex/Source/Private/ReflexStatsLatencyMarkers.cpp
@@ -1,6 +1,7 @@
 // Copyright Epic Games, Inc. All Rights Reserved.
 
 #include "ReflexStatsLatencyMarkers.h"
+#include "Framework/Application/SlateApplication.h"
 
 THIRD_PARTY_INCLUDES_START
 __pragma(warning(disable:4191)) // hide a warning from TraceLoggingProvider.h
