VulkanRHI: Disable the Reflex plugin on Vulkan.

#rnx

diff --git a/Engine/Plugins/Runtime/Nvidia/Reflex/Source/Private/ReflexLatencyMarkers.cpp b/Engine/Plugins/Runtime/Nvidia/Reflex/Source/Private/ReflexLatencyMarkers.cpp
--- a/Engine/Plugins/Runtime/Nvidia/Reflex/Source/Private/ReflexLatencyMarkers.cpp
+++ b/Engine/Plugins/Runtime/Nvidia/Reflex/Source/Private/ReflexLatencyMarkers.cpp
@@ -30,6 +30,12 @@ void FReflexLatencyMarkers::Initialize()
 {
 	if (IsRHIDeviceNVIDIA())
 	{
+		FString RHIName = GDynamicRHI->GetName();
+		if (RHIName.StartsWith(TEXT("Vulkan")))
+		{
+			return;
+		}
+
 		NvU32 DriverVersion = 0;
 		NvAPI_ShortString BranchString;
 
diff --git a/Engine/Plugins/Runtime/Nvidia/Reflex/Source/Private/ReflexMaxTickRateHandler.cpp b/Engine/Plugins/Runtime/Nvidia/Reflex/Source/Private/ReflexMaxTickRateHandler.cpp
--- a/Engine/Plugins/Runtime/Nvidia/Reflex/Source/Private/ReflexMaxTickRateHandler.cpp
+++ b/Engine/Plugins/Runtime/Nvidia/Reflex/Source/Private/ReflexMaxTickRateHandler.cpp
@@ -44,6 +44,12 @@ void FReflexMaxTickRateHandler::Initialize()
 {
 	if (IsRHIDeviceNVIDIA())
 	{
+		FString RHIName = GDynamicRHI->GetName();
+		if (RHIName.StartsWith(TEXT("Vulkan")))
+		{
+			return;
+		}
+
 		NvU32 DriverVersion;
 		NvAPI_ShortString BranchString;
 
