Integrate fix from Nvidia on UReflexBlueprintLibrary

diff --git a/Engine/Plugins/Runtime/Nvidia/Reflex/Source/Private/ReflexBlueprint.cpp b/Engine/Plugins/Runtime/Nvidia/Reflex/Source/Private/ReflexBlueprint.cpp
--- a/Engine/Plugins/Runtime/Nvidia/Reflex/Source/Private/ReflexBlueprint.cpp
+++ b/Engine/Plugins/Runtime/Nvidia/Reflex/Source/Private/ReflexBlueprint.cpp
@@ -92,7 +92,9 @@
 		.GetModularFeatureImplementations<ILatencyMarkerModule>(ILatencyMarkerModule::GetModularFeatureName());
 	for (ILatencyMarkerModule* LatencyMarkerModule : LatencyMarkerModules)
 	{
-		return LatencyMarkerModule->GetTotalLatencyInMs();
+		const float Latency = LatencyMarkerModule->GetTotalLatencyInMs();
+		if (Latency > 0.f)
+			return Latency;
 	}
 
 	return 0.f;
@@ -104,7 +106,9 @@
 		.GetModularFeatureImplementations<ILatencyMarkerModule>(ILatencyMarkerModule::GetModularFeatureName());
 	for (ILatencyMarkerModule* LatencyMarkerModule : LatencyMarkerModules)
 	{
-		return LatencyMarkerModule->GetGameLatencyInMs();
+		const float Latency = LatencyMarkerModule->GetGameLatencyInMs();
+		if (Latency > 0.f)
+			return Latency;
 	}
 
 	return 0.f;
@@ -116,7 +120,9 @@
 		.GetModularFeatureImplementations<ILatencyMarkerModule>(ILatencyMarkerModule::GetModularFeatureName());
 	for (ILatencyMarkerModule* LatencyMarkerModule : LatencyMarkerModules)
 	{
-		return LatencyMarkerModule->GetRenderLatencyInMs();
+		const float Latency = LatencyMarkerModule->GetRenderLatencyInMs();
+		if (Latency > 0.f)
+			return Latency;
 	}
 
 	return 0.f;
