PR #9275: Export ReflexBlueprintLibrary to other modules (Contributed by [REDACTED_PERSON])

diff --git a/Engine/Plugins/Runtime/Nvidia/Reflex/Source/Public/ReflexBlueprint.h b/Engine/Plugins/Runtime/Nvidia/Reflex/Source/Public/ReflexBlueprint.h
--- a/Engine/Plugins/Runtime/Nvidia/Reflex/Source/Public/ReflexBlueprint.h
+++ b/Engine/Plugins/Runtime/Nvidia/Reflex/Source/Public/ReflexBlueprint.h
@@ -16,7 +16,7 @@ enum class EReflexMode : uint8
 
 
 UCLASS()
-class UReflexBlueprintLibrary : public UBlueprintFunctionLibrary
+class REFLEX_API UReflexBlueprintLibrary : public UBlueprintFunctionLibrary
 {
 	GENERATED_UCLASS_BODY()
 
