FExec: Converted Plugins Exec overrides to Exec_Runtime/Exec_Dev where:
- No additional routing was being done in Exec.
- No default parameter was being provided through the Exec override.

This is so we can turn off Exec overrides in certain configurations.

diff --git a/Engine/Plugins/MemoryUsageQueries/Source/MemoryUsageQueries/Private/MemoryUsageQueries.cpp b/Engine/Plugins/MemoryUsageQueries/Source/MemoryUsageQueries/Private/MemoryUsageQueries.cpp
--- a/Engine/Plugins/MemoryUsageQueries/Source/MemoryUsageQueries/Private/MemoryUsageQueries.cpp
+++ b/Engine/Plugins/MemoryUsageQueries/Source/MemoryUsageQueries/Private/MemoryUsageQueries.cpp
@@ -30,7 +30,8 @@ struct FMemoryUsageQueriesExec : public FSelfRegisteringExec
 {
 	FMemoryUsageQueriesExec() {}
 
-	virtual bool Exec(UWorld* Inworld, const TCHAR* Cmd, FOutputDevice& Ar) override;
+protected:
+	virtual bool Exec_Runtime(UWorld* Inworld, const TCHAR* Cmd, FOutputDevice& Ar) override;
 };
 
 struct FAssetMemoryBreakdown
@@ -75,7 +76,7 @@ struct FAssetMemoryDetails
 static FMemoryUsageQueriesExec DebugSettinsSubsystemExecInstance;
 static int32 DefaultResultLimit = 15;
 
-bool FMemoryUsageQueriesExec::Exec(UWorld* InWorld, const TCHAR* Cmd, FOutputDevice& Ar)
+bool FMemoryUsageQueriesExec::Exec_Runtime(UWorld* InWorld, const TCHAR* Cmd, FOutputDevice& Ar)
 {
 	if (FParse::Command(&Cmd, TEXT("MemQuery")))
 	{
diff --git a/Engine/Plugins/Messaging/TcpMessaging/Source/TcpMessaging/Private/TcpMessagingModule.cpp b/Engine/Plugins/Messaging/TcpMessaging/Source/TcpMessaging/Private/TcpMessagingModule.cpp
--- a/Engine/Plugins/Messaging/TcpMessaging/Source/TcpMessaging/Private/TcpMessagingModule.cpp
+++ b/Engine/Plugins/Messaging/TcpMessaging/Source/TcpMessaging/Private/TcpMessagingModule.cpp
@@ -28,11 +28,11 @@ class FTcpMessagingModule
 	: public FSelfRegisteringExec
 	, public ITcpMessagingModule
 {
-public:
+protected:
 
 	// FSelfRegisteringExec interface
 
-	virtual bool Exec(UWorld* InWorld, const TCHAR* Cmd, FOutputDevice& Ar) override
+	virtual bool Exec_Runtime(UWorld* InWorld, const TCHAR* Cmd, FOutputDevice& Ar) override
 	{
 		if (!FParse::Command(&Cmd, TEXT("TCPMESSAGING")))
 		{
diff --git a/Engine/Plugins/Messaging/UdpMessaging/Source/UdpMessaging/Private/UdpMessagingModule.cpp b/Engine/Plugins/Messaging/UdpMessaging/Source/UdpMessaging/Private/UdpMessagingModule.cpp
--- a/Engine/Plugins/Messaging/UdpMessaging/Source/UdpMessaging/Private/UdpMessagingModule.cpp
+++ b/Engine/Plugins/Messaging/UdpMessaging/Source/UdpMessaging/Private/UdpMessagingModule.cpp
@@ -33,11 +33,11 @@ class FUdpMessagingModule
 	, public INetworkMessagingExtension
 	, public IModuleInterface
 {
-public:
+protected:
 
 	//~ FSelfRegisteringExec interface
 
-	virtual bool Exec(UWorld* InWorld, const TCHAR* Cmd, FOutputDevice& Ar) override
+	virtual bool Exec_Runtime(UWorld* InWorld, const TCHAR* Cmd, FOutputDevice& Ar) override
 	{
 		if (!FParse::Command(&Cmd, TEXT("UDPMESSAGING")))
 		{
diff --git a/Engine/Plugins/Online/EOSShared/Source/EOSShared/Private/EOSSDKManager.cpp b/Engine/Plugins/Online/EOSShared/Source/EOSShared/Private/EOSSDKManager.cpp
--- a/Engine/Plugins/Online/EOSShared/Source/EOSShared/Private/EOSSDKManager.cpp
+++ b/Engine/Plugins/Online/EOSShared/Source/EOSShared/Private/EOSSDKManager.cpp
@@ -710,7 +710,7 @@ EOS_EResult FEOSSDKManager::EOSInitialize(EOS_InitializeOptions& Options)
 	return EOS_Initialize(&Options);
 }
 
-bool FEOSSDKManager::Exec(class UWorld* InWorld, const TCHAR* Cmd, FOutputDevice& Ar)
+bool FEOSSDKManager::Exec_Runtime(class UWorld* InWorld, const TCHAR* Cmd, FOutputDevice& Ar)
 {
 	if (!FParse::Command(&Cmd, TEXT("EOSSDK")))
 	{
diff --git a/Engine/Plugins/Online/EOSShared/Source/EOSShared/Private/EOSSDKManager.h b/Engine/Plugins/Online/EOSShared/Source/EOSShared/Private/EOSSDKManager.h
--- a/Engine/Plugins/Online/EOSShared/Source/EOSShared/Private/EOSSDKManager.h
+++ b/Engine/Plugins/Online/EOSShared/Source/EOSShared/Private/EOSSDKManager.h
@@ -57,13 +57,13 @@ public:
 	virtual void AddCallbackObject(TUniquePtr<class FCallbackBase> CallbackObj) override;
 	// End IEOSSDKManager
 
-	// Begin FSelfRegisteringExec
-	virtual bool Exec(class UWorld* InWorld, const TCHAR* Cmd, FOutputDevice& Ar) override;
-	// End FSelfRegisteringExec
-
 	void Shutdown();
 
 protected:
+	// Begin FSelfRegisteringExec
+	virtual bool Exec_Runtime(class UWorld* InWorld, const TCHAR* Cmd, FOutputDevice& Ar) override;
+	// End FSelfRegisteringExec
+
 	virtual EOS_EResult EOSInitialize(EOS_InitializeOptions& Options);
 	virtual IEOSPlatformHandlePtr CreatePlatform(const FEOSSDKPlatformConfig& PlatformConfig, EOS_Platform_Options& PlatformOptions);
 	virtual bool Tick(float);
diff --git a/Engine/Plugins/Online/OnlineFramework/Source/Hotfix/Private/OnlineHotfixManager.cpp b/Engine/Plugins/Online/OnlineFramework/Source/Hotfix/Private/OnlineHotfixManager.cpp
--- a/Engine/Plugins/Online/OnlineFramework/Source/Hotfix/Private/OnlineHotfixManager.cpp
+++ b/Engine/Plugins/Online/OnlineFramework/Source/Hotfix/Private/OnlineHotfixManager.cpp
@@ -1958,7 +1958,8 @@ UWorld* UOnlineHotfixManager::GetWorld() const
 struct FHotfixManagerExec :
 	public FSelfRegisteringExec
 {
-	virtual bool Exec(UWorld* InWorld, const TCHAR* Cmd, FOutputDevice& Ar) override
+protected:
+	virtual bool Exec_Runtime(UWorld* InWorld, const TCHAR* Cmd, FOutputDevice& Ar) override
 	{
 		if (FParse::Command(&Cmd, TEXT("HOTFIX")))
 		{
diff --git a/Engine/Plugins/Online/OnlineFramework/Source/PlayTimeLimit/Private/PlayTimeLimitModule.cpp b/Engine/Plugins/Online/OnlineFramework/Source/PlayTimeLimit/Private/PlayTimeLimitModule.cpp
--- a/Engine/Plugins/Online/OnlineFramework/Source/PlayTimeLimit/Private/PlayTimeLimitModule.cpp
+++ b/Engine/Plugins/Online/OnlineFramework/Source/PlayTimeLimit/Private/PlayTimeLimitModule.cpp
@@ -29,7 +29,7 @@ void FPlayTimeLimitModule::ShutdownModule()
 	}
 }
 
-bool FPlayTimeLimitModule::Exec(UWorld* InWorld, const TCHAR* Cmd, FOutputDevice& Ar)
+bool FPlayTimeLimitModule::Exec_Runtime(UWorld* InWorld, const TCHAR* Cmd, FOutputDevice& Ar)
 {
 	bool bWasHandled = false;
 
diff --git a/Engine/Plugins/Online/OnlineFramework/Source/PlayTimeLimit/Public/PlayTimeLimitModule.h b/Engine/Plugins/Online/OnlineFramework/Source/PlayTimeLimit/Public/PlayTimeLimitModule.h
--- a/Engine/Plugins/Online/OnlineFramework/Source/PlayTimeLimit/Public/PlayTimeLimitModule.h
+++ b/Engine/Plugins/Online/OnlineFramework/Source/PlayTimeLimit/Public/PlayTimeLimitModule.h
@@ -14,9 +14,6 @@ class FPlayTimeLimitModule :
 {
 
 public:
-	// FSelfRegisteringExec
-	virtual bool Exec(UWorld* InWorld, const TCHAR* Cmd, FOutputDevice& Ar) override;
-
 	/**
 	 * Singleton-like access to this module's interface.  This is just for convenience!
 	 * Beware of calling this during the shutdown phase, though.  Your module might have been unloaded already.
@@ -38,6 +35,10 @@ public:
 		return FModuleManager::Get().IsModuleLoaded("PlayTimeLimit");
 	}
 
+protected:
+	// FSelfRegisteringExec
+	virtual bool Exec_Runtime(UWorld* InWorld, const TCHAR* Cmd, FOutputDevice& Ar) override;
+
 private:
 
 	// IModuleInterface
diff --git a/Engine/Plugins/Online/OnlineFramework/Source/Qos/Private/QosModule.cpp b/Engine/Plugins/Online/OnlineFramework/Source/Qos/Private/QosModule.cpp
--- a/Engine/Plugins/Online/OnlineFramework/Source/Qos/Private/QosModule.cpp
+++ b/Engine/Plugins/Online/OnlineFramework/Source/Qos/Private/QosModule.cpp
@@ -28,7 +28,7 @@ TSharedRef<FQosInterface> FQosModule::GetQosInterface()
 	return QosInterface.ToSharedRef();
 }
 
-bool FQosModule::Exec(UWorld* InWorld, const TCHAR* Cmd, FOutputDevice& Ar)
+bool FQosModule::Exec_Runtime(UWorld* InWorld, const TCHAR* Cmd, FOutputDevice& Ar)
 {
 	bool bWasHandled = false;
 
diff --git a/Engine/Plugins/Online/OnlineFramework/Source/Qos/Public/QosModule.h b/Engine/Plugins/Online/OnlineFramework/Source/Qos/Public/QosModule.h
--- a/Engine/Plugins/Online/OnlineFramework/Source/Qos/Public/QosModule.h
+++ b/Engine/Plugins/Online/OnlineFramework/Source/Qos/Public/QosModule.h
@@ -18,9 +18,6 @@ class QOS_API FQosModule :
 
 public:
 
-	// FSelfRegisteringExec
-	virtual bool Exec(UWorld* InWorld, const TCHAR* Cmd, FOutputDevice& Ar) override;
-
 	/**
 	 * Singleton-like access to this module's interface.  This is just for convenience!
 	 * Beware of calling this during the shutdown phase, though.  Your module might have been unloaded already.
@@ -47,6 +44,10 @@ public:
 	 */
 	TSharedRef<FQosInterface> GetQosInterface();
 
+protected:
+	// FSelfRegisteringExec
+	virtual bool Exec_Runtime(UWorld* InWorld, const TCHAR* Cmd, FOutputDevice& Ar) override;
+
 private:
 	TSharedPtr<FQosInterface> QosInterface;
 
diff --git a/Engine/Plugins/Online/OnlineSubsystemUtils/Source/OnlineSubsystemUtils/Private/Tests/TestSessionInterface.cpp b/Engine/Plugins/Online/OnlineSubsystemUtils/Source/OnlineSubsystemUtils/Private/Tests/TestSessionInterface.cpp
--- a/Engine/Plugins/Online/OnlineSubsystemUtils/Source/OnlineSubsystemUtils/Private/Tests/TestSessionInterface.cpp
+++ b/Engine/Plugins/Online/OnlineSubsystemUtils/Source/OnlineSubsystemUtils/Private/Tests/TestSessionInterface.cpp
@@ -455,7 +455,7 @@ bool FTestSessionInterface::Tick(float DeltaTime)
 	return true;
 }
 
-bool FTestSessionInterface::Exec(UWorld* InWorld, const TCHAR* Cmd, FOutputDevice& Ar)
+bool FTestSessionInterface::Exec_Runtime(UWorld* InWorld, const TCHAR* Cmd, FOutputDevice& Ar)
 {
 	bool bWasHandled = false;
 
diff --git a/Engine/Plugins/Online/OnlineSubsystemUtils/Source/OnlineSubsystemUtils/Private/Tests/TestSessionInterface.h b/Engine/Plugins/Online/OnlineSubsystemUtils/Source/OnlineSubsystemUtils/Private/Tests/TestSessionInterface.h
--- a/Engine/Plugins/Online/OnlineSubsystemUtils/Source/OnlineSubsystemUtils/Private/Tests/TestSessionInterface.h
+++ b/Engine/Plugins/Online/OnlineSubsystemUtils/Source/OnlineSubsystemUtils/Private/Tests/TestSessionInterface.h
@@ -344,9 +344,6 @@ class TestOnlineSearchSettings;
 	// FTSTickerObjectBase
 	bool Tick(float DeltaTime) override;
 
-	// FSelfRegisteringExec
-	virtual bool Exec(UWorld* InWorld, const TCHAR* Cmd, FOutputDevice& Ar) override;
-
 	// FTestSessionInterface
 
 	/**
@@ -360,6 +357,10 @@ class TestOnlineSearchSettings;
 	 * Clear out any existing delegates created
 	 */
 	void ClearDelegates();
+
+protected:
+	// FSelfRegisteringExec
+	virtual bool Exec_Runtime(UWorld* InWorld, const TCHAR* Cmd, FOutputDevice& Ar) override;
  };
 
 #endif //WITH_DEV_AUTOMATION_TESTS
diff --git a/Engine/Plugins/Online/OnlineSubsystemUtils/Source/OnlineSubsystemUtils/Private/Tests/TestVoice.cpp b/Engine/Plugins/Online/OnlineSubsystemUtils/Source/OnlineSubsystemUtils/Private/Tests/TestVoice.cpp
--- a/Engine/Plugins/Online/OnlineSubsystemUtils/Source/OnlineSubsystemUtils/Private/Tests/TestVoice.cpp
+++ b/Engine/Plugins/Online/OnlineSubsystemUtils/Source/OnlineSubsystemUtils/Private/Tests/TestVoice.cpp
@@ -397,7 +397,7 @@ bool FTestVoice::Tick(float DeltaTime)
 	return true;
 }
 
-bool FTestVoice::Exec(UWorld* InWorld, const TCHAR* Cmd, FOutputDevice& Ar)
+bool FTestVoice::Exec_Runtime(UWorld* InWorld, const TCHAR* Cmd, FOutputDevice& Ar)
 {
 	bool bWasHandled = false;
 	if (FParse::Command(&Cmd, TEXT("killtestvoice")))
diff --git a/Engine/Plugins/Online/OnlineSubsystemUtils/Source/OnlineSubsystemUtils/Private/Tests/TestVoice.h b/Engine/Plugins/Online/OnlineSubsystemUtils/Source/OnlineSubsystemUtils/Private/Tests/TestVoice.h
--- a/Engine/Plugins/Online/OnlineSubsystemUtils/Source/OnlineSubsystemUtils/Private/Tests/TestVoice.h
+++ b/Engine/Plugins/Online/OnlineSubsystemUtils/Source/OnlineSubsystemUtils/Private/Tests/TestVoice.h
@@ -27,11 +27,12 @@ public:
 	// FTSTickerObjectBase
 	virtual bool Tick(float DeltaTime) override;
 
-	// FSelfRegisteringExec
-	virtual bool Exec(UWorld* InWorld, const TCHAR* Cmd, FOutputDevice& Ar) override;
-
 	void Test();
 
+protected:
+	// FSelfRegisteringExec
+	virtual bool Exec_Runtime(UWorld* InWorld, const TCHAR* Cmd, FOutputDevice& Ar) override;
+
 private:
 
 	/** Reference to audio component playing active streaming sound */
diff --git a/Engine/Plugins/Online/OnlineSubsystemUtils/Source/OnlineSubsystemUtils/Private/VoiceEngineImpl.cpp b/Engine/Plugins/Online/OnlineSubsystemUtils/Source/OnlineSubsystemUtils/Private/VoiceEngineImpl.cpp
--- a/Engine/Plugins/Online/OnlineSubsystemUtils/Source/OnlineSubsystemUtils/Private/VoiceEngineImpl.cpp
+++ b/Engine/Plugins/Online/OnlineSubsystemUtils/Source/OnlineSubsystemUtils/Private/VoiceEngineImpl.cpp
@@ -736,7 +736,7 @@ FString FVoiceEngineImpl::GetVoiceDebugState() const
 	return Output;
 }
 
-bool FVoiceEngineImpl::Exec(UWorld* InWorld, const TCHAR* Cmd, FOutputDevice& Ar)
+bool FVoiceEngineImpl::Exec_Runtime(UWorld* InWorld, const TCHAR* Cmd, FOutputDevice& Ar)
 {
 	bool bWasHandled = false;
 
diff --git a/Engine/Plugins/Online/OnlineSubsystemUtils/Source/OnlineSubsystemUtils/Public/VoiceEngineImpl.h b/Engine/Plugins/Online/OnlineSubsystemUtils/Source/OnlineSubsystemUtils/Public/VoiceEngineImpl.h
--- a/Engine/Plugins/Online/OnlineSubsystemUtils/Source/OnlineSubsystemUtils/Public/VoiceEngineImpl.h
+++ b/Engine/Plugins/Online/OnlineSubsystemUtils/Source/OnlineSubsystemUtils/Public/VoiceEngineImpl.h
@@ -331,9 +331,6 @@ public:
 	virtual void Tick(float DeltaTime) override;
 	FString GetVoiceDebugState() const override;
 
-	// FSelfRegisteringExec
-	virtual bool Exec(UWorld* InWorld, const TCHAR* Cmd, FOutputDevice& Ar) override;
-
 	virtual void GetVoiceSettingsOverride(const FUniqueNetIdWrapper& RemoteTalkerId, FVoiceSettings& VoiceSettings) {}
 
 
@@ -352,6 +349,10 @@ public:
 
 	virtual bool PatchLocalTalkerOutputToEndpoint(const FString& InDeviceName) override;
 
+protected:
+	// FSelfRegisteringExec
+	virtual bool Exec_Runtime(UWorld* InWorld, const TCHAR* Cmd, FOutputDevice& Ar) override;
+
 private:
 
 	/**
diff --git a/Engine/Plugins/Runtime/AR/AppleAR/AppleARKit/Source/AppleARKit/Private/AppleARKitSystem.cpp b/Engine/Plugins/Runtime/AR/AppleAR/AppleARKit/Source/AppleARKit/Private/AppleARKitSystem.cpp
--- a/Engine/Plugins/Runtime/AR/AppleAR/AppleARKit/Source/AppleARKit/Private/AppleARKitSystem.cpp
+++ b/Engine/Plugins/Runtime/AR/AppleAR/AppleARKit/Source/AppleARKit/Private/AppleARKitSystem.cpp
@@ -2891,7 +2891,7 @@ int32 UAppleARKitSettings::GetRenderThreadPriorityOverride()
 	return RenderThreadPriorityOverride;
 }
 
-bool UAppleARKitSettings::Exec(UWorld*, const TCHAR* Cmd, FOutputDevice& Ar)
+bool UAppleARKitSettings::Exec_Runtime(UWorld*, const TCHAR* Cmd, FOutputDevice& Ar)
 {
 	if (FParse::Command(&Cmd, TEXT("ARKitSettings")))
 	{
diff --git a/Engine/Plugins/Runtime/AR/AppleAR/AppleARKit/Source/AppleARKit/Public/AppleARKitSettings.h b/Engine/Plugins/Runtime/AR/AppleAR/AppleARKit/Source/AppleARKit/Public/AppleARKitSettings.h
--- a/Engine/Plugins/Runtime/AR/AppleAR/AppleARKit/Source/AppleARKit/Public/AppleARKitSettings.h
+++ b/Engine/Plugins/Runtime/AR/AppleAR/AppleARKit/Source/AppleARKit/Public/AppleARKitSettings.h
@@ -85,7 +85,7 @@ public:
 
 protected:
 	//~ FSelfRegisteringExec
-	virtual bool Exec(UWorld*, const TCHAR* Cmd, FOutputDevice& Ar) override;
+	virtual bool Exec_Runtime(UWorld*, const TCHAR* Cmd, FOutputDevice& Ar) override;
 	//~ FSelfRegisteringExec
 
 	/** Where to write the curve files and image files */
diff --git a/Engine/Plugins/Runtime/AR/AppleAR/AppleARKitFaceSupport/Source/AppleARKitFaceSupport/Private/AppleARKitFaceSupportImpl.cpp b/Engine/Plugins/Runtime/AR/AppleAR/AppleARKitFaceSupport/Source/AppleARKitFaceSupport/Private/AppleARKitFaceSupportImpl.cpp
--- a/Engine/Plugins/Runtime/AR/AppleAR/AppleARKitFaceSupport/Source/AppleARKitFaceSupport/Private/AppleARKitFaceSupportImpl.cpp
+++ b/Engine/Plugins/Runtime/AR/AppleAR/AppleARKitFaceSupport/Source/AppleARKitFaceSupport/Private/AppleARKitFaceSupportImpl.cpp
@@ -103,7 +103,7 @@ void FAppleARKitFaceSupport::Shutdown()
 	IModularFeatures::Get().UnregisterModularFeature(GetModularFeatureName(), this);
 }
 
-bool FAppleARKitFaceSupport::Exec(UWorld*, const TCHAR* Cmd, FOutputDevice& Ar)
+bool FAppleARKitFaceSupport::Exec_Runtime(UWorld*, const TCHAR* Cmd, FOutputDevice& Ar)
 {
 	if (FParse::Command(&Cmd, TEXT("LiveLinkFaceAR")))
 	{
diff --git a/Engine/Plugins/Runtime/AR/AppleAR/AppleARKitFaceSupport/Source/AppleARKitFaceSupport/Private/AppleARKitFaceSupportImpl.h b/Engine/Plugins/Runtime/AR/AppleAR/AppleARKitFaceSupport/Source/AppleARKitFaceSupport/Private/AppleARKitFaceSupportImpl.h
--- a/Engine/Plugins/Runtime/AR/AppleAR/AppleARKitFaceSupport/Source/AppleARKitFaceSupport/Private/AppleARKitFaceSupportImpl.h
+++ b/Engine/Plugins/Runtime/AR/AppleAR/AppleARKitFaceSupport/Source/AppleARKitFaceSupport/Private/AppleARKitFaceSupportImpl.h
@@ -52,14 +52,14 @@ private:
 	
 	FName GetLiveLinkSubjectName(const FGuid& AnchorId);
 
-	//~ FSelfRegisteringExec
-	virtual bool Exec(UWorld*, const TCHAR* Cmd, FOutputDevice& Ar) override;
-	//~ FSelfRegisteringExec
-
 	/** Whether the face data is mirrored or not */
 	bool bFaceMirrored;
 	
 protected:
+	//~ FSelfRegisteringExec
+	virtual bool Exec_Runtime(UWorld*, const TCHAR* Cmd, FOutputDevice& Ar) override;
+	//~ FSelfRegisteringExec
+
 	/** If requested, publishes face ar updates to LiveLink for the animation system to use */
 	TSharedPtr<ILiveLinkSourceARKit> LiveLinkSource;
 	/** The id of this device */
diff --git a/Engine/Plugins/Runtime/AR/AppleAR/AppleARKitFaceSupport/Source/AppleARKitFaceSupport/Private/AppleARKitLiveLinkSource.cpp b/Engine/Plugins/Runtime/AR/AppleAR/AppleARKitFaceSupport/Source/AppleARKitFaceSupport/Private/AppleARKitLiveLinkSource.cpp
--- a/Engine/Plugins/Runtime/AR/AppleAR/AppleARKitFaceSupport/Source/AppleARKitFaceSupport/Private/AppleARKitLiveLinkSource.cpp
+++ b/Engine/Plugins/Runtime/AR/AppleAR/AppleARKitFaceSupport/Source/AppleARKitFaceSupport/Private/AppleARKitLiveLinkSource.cpp
@@ -538,7 +538,7 @@ void FAppleARKitLiveLinkFileWriter::PublishBlendShapes(FName SubjectName, const
 	}
 }
 
-bool FAppleARKitLiveLinkFileWriter::Exec(UWorld*, const TCHAR* Cmd, FOutputDevice& Ar)
+bool FAppleARKitLiveLinkFileWriter::Exec_Runtime(UWorld*, const TCHAR* Cmd, FOutputDevice& Ar)
 {
 	if (FParse::Command(&Cmd, TEXT("FaceAR")))
 	{
diff --git a/Engine/Plugins/Runtime/AR/AppleAR/AppleARKitFaceSupport/Source/AppleARKitFaceSupport/Private/AppleARKitLiveLinkSource.h b/Engine/Plugins/Runtime/AR/AppleAR/AppleARKitFaceSupport/Source/AppleARKitFaceSupport/Private/AppleARKitLiveLinkSource.h
--- a/Engine/Plugins/Runtime/AR/AppleAR/AppleARKitFaceSupport/Source/AppleARKitFaceSupport/Private/AppleARKitLiveLinkSource.h
+++ b/Engine/Plugins/Runtime/AR/AppleAR/AppleARKitFaceSupport/Source/AppleARKitFaceSupport/Private/AppleARKitLiveLinkSource.h
@@ -810,7 +810,7 @@ protected:
 	// End IARKitBlendShapePublisher
 
 	//~ FSelfRegisteringExec
-	virtual bool Exec(UWorld*, const TCHAR* Cmd, FOutputDevice& Ar) override;
+	virtual bool Exec_Runtime(UWorld*, const TCHAR* Cmd, FOutputDevice& Ar) override;
 	//~ FSelfRegisteringExec
 
 	virtual FString BuildSaveData() = 0;
diff --git a/Engine/Plugins/Runtime/Nvidia/Reflex/Source/Private/ReflexLatencyMarkers.cpp b/Engine/Plugins/Runtime/Nvidia/Reflex/Source/Private/ReflexLatencyMarkers.cpp
--- a/Engine/Plugins/Runtime/Nvidia/Reflex/Source/Private/ReflexLatencyMarkers.cpp
+++ b/Engine/Plugins/Runtime/Nvidia/Reflex/Source/Private/ReflexLatencyMarkers.cpp
@@ -310,7 +310,7 @@ bool FReflexLatencyMarkers::GetAvailable()
 	return true;
 }
 
-bool FReflexLatencyMarkers::Exec(UWorld* InWorld, const TCHAR* Cmd, FOutputDevice& Ar)
+bool FReflexLatencyMarkers::Exec_Runtime(UWorld* InWorld, const TCHAR* Cmd, FOutputDevice& Ar)
 {
 	bool bHandled = false;
 	FString ReflexLatencyMarkers;
@@ -346,4 +346,4 @@ bool FReflexLatencyMarkers::Exec(UWorld* InWorld, const TCHAR* Cmd, FOutputDevic
 	}
 
 	return bHandled;
-}
\ No newline at end of file
+}
diff --git a/Engine/Plugins/Runtime/Nvidia/Reflex/Source/Private/ReflexMaxTickRateHandler.cpp b/Engine/Plugins/Runtime/Nvidia/Reflex/Source/Private/ReflexMaxTickRateHandler.cpp
--- a/Engine/Plugins/Runtime/Nvidia/Reflex/Source/Private/ReflexMaxTickRateHandler.cpp
+++ b/Engine/Plugins/Runtime/Nvidia/Reflex/Source/Private/ReflexMaxTickRateHandler.cpp
@@ -191,7 +191,7 @@ bool FReflexMaxTickRateHandler::GetAvailable()
 	return true;
 }
 
-bool FReflexMaxTickRateHandler::Exec(UWorld* InWorld, const TCHAR* Cmd, FOutputDevice& Ar)
+bool FReflexMaxTickRateHandler::Exec_Runtime(UWorld* InWorld, const TCHAR* Cmd, FOutputDevice& Ar)
 {
 	bool bHandled = false;
 	FString ReflexMode;
@@ -270,4 +270,4 @@ bool FReflexMaxTickRateHandler::Exec(UWorld* InWorld, const TCHAR* Cmd, FOutputD
 	}
 
 	return bHandled;
-}
\ No newline at end of file
+}
diff --git a/Engine/Plugins/Runtime/Nvidia/Reflex/Source/Public/ReflexLatencyMarkers.h b/Engine/Plugins/Runtime/Nvidia/Reflex/Source/Public/ReflexLatencyMarkers.h
--- a/Engine/Plugins/Runtime/Nvidia/Reflex/Source/Public/ReflexLatencyMarkers.h
+++ b/Engine/Plugins/Runtime/Nvidia/Reflex/Source/Public/ReflexLatencyMarkers.h
@@ -50,8 +50,6 @@ public:
 	virtual void SetFlashIndicatorEnabled(bool bInEnabled) override;
 	virtual bool GetFlashIndicatorEnabled() override;
 
-	virtual bool Exec(UWorld* InWorld, const TCHAR* Cmd, FOutputDevice& Ar) override;
-
 	virtual void SetInputSampleLatencyMarker(uint64 FrameNumber) override;
 	virtual void SetSimulationLatencyMarkerStart(uint64 FrameNumber) override;
 	virtual void SetSimulationLatencyMarkerEnd(uint64 FrameNumber) override;
@@ -78,4 +76,7 @@ public:
 	virtual float GetDriverOffsetFromFrameStartInMs() override { return DriverOffsetMs; }
 	virtual float GetOSRenderQueueOffsetFromFrameStartInMs() override { return OSRenderQueueOffsetMs; }
 	virtual float GetGPURenderOffsetFromFrameStartInMs() override { return GPURenderOffsetMs; }
+
+protected:
+	virtual bool Exec_Runtime(UWorld* InWorld, const TCHAR* Cmd, FOutputDevice& Ar) override;
 };
\ No newline at end of file
diff --git a/Engine/Plugins/Runtime/Nvidia/Reflex/Source/Public/ReflexMaxTickRateHandler.h b/Engine/Plugins/Runtime/Nvidia/Reflex/Source/Public/ReflexMaxTickRateHandler.h
--- a/Engine/Plugins/Runtime/Nvidia/Reflex/Source/Public/ReflexMaxTickRateHandler.h
+++ b/Engine/Plugins/Runtime/Nvidia/Reflex/Source/Public/ReflexMaxTickRateHandler.h
@@ -16,8 +16,6 @@ public:
 	virtual bool GetEnabled() override;
 	virtual bool GetAvailable() override;
 
-	virtual bool Exec(UWorld* InWorld, const TCHAR* Cmd, FOutputDevice& Ar) override;
-
 	// Used to provide a generic customization interface for custom tick rate handlers
 	virtual void SetFlags(uint32 Flags);
 	virtual uint32 GetFlags();
@@ -34,4 +32,7 @@ public:
 
 	uint32 CustomFlags = 0;
 	uint32 LastCustomFlags = 0;
+
+protected:
+	virtual bool Exec_Runtime(UWorld* InWorld, const TCHAR* Cmd, FOutputDevice& Ar) override;
 };
\ No newline at end of file
