Reflex fix to allow PC latency markers to work with non-NVIDIA hardware
The flash indicator sub-feature still requires a minimum NVIDIA driver

diff --git a/Engine/Plugins/Runtime/Nvidia/Reflex/Source/Private/ReflexLatencyMarkers.cpp b/Engine/Plugins/Runtime/Nvidia/Reflex/Source/Private/ReflexLatencyMarkers.cpp
--- a/Engine/Plugins/Runtime/Nvidia/Reflex/Source/Private/ReflexLatencyMarkers.cpp
+++ b/Engine/Plugins/Runtime/Nvidia/Reflex/Source/Private/ReflexLatencyMarkers.cpp
@@ -29,46 +29,31 @@ static FAutoConsoleVariableRef CVarDisableLatencyMarkers(
 
 void FReflexLatencyMarkers::Initialize()
 {
-	if ((RHIGetInterfaceType() == ERHIInterfaceType::D3D11 || RHIGetInterfaceType() == ERHIInterfaceType::D3D12) && IsRHIDeviceNVIDIA())
+	if ((RHIGetInterfaceType() == ERHIInterfaceType::D3D11 || RHIGetInterfaceType() == ERHIInterfaceType::D3D12))
 	{
-		FString RHIName = GDynamicRHI->GetName();
-		if (RHIName.StartsWith(TEXT("Vulkan")))
-		{
-			return;
-		}
+		// No vendor or driver version required for basic Latency Marker support.
+		bFeatureSupport = true;
+		bFlashIndicatorDriverControlled = false;
 
-		NvU32 DriverVersion = 0;
-		NvAPI_ShortString BranchString;
-
-		// Driver version check, 455 and above required for Reflex
-		NvAPI_SYS_GetDriverAndBranchVersion(&DriverVersion, BranchString);
-		if (DriverVersion >= 45500)
+		// Check for flash indicator sub-feature support, which still requires
+		// a minimum NVIDIA driver version.
+		if (IsRHIDeviceNVIDIA())
 		{
-			bProperDriverVersion = true;
-		}
-
-		if (DriverVersion >= 51123)
-		{
-			bFlashIndicatorDriverControlled = true;
-		}
+			NvU32 DriverVersion = 0;
+			NvAPI_ShortString BranchString;
+			NvAPI_SYS_GetDriverAndBranchVersion(&DriverVersion, BranchString);
 
-		// Need to verify that Reflex Low Latency mode is supported on current hardware
-		NV_GET_SLEEP_STATUS_PARAMS_V1 SleepStatusParams = { 0 };
-		SleepStatusParams.version = NV_GET_SLEEP_STATUS_PARAMS_VER1;
-
-		NvAPI_Status status = NVAPI_OK;
-		status = NvAPI_D3D_GetSleepStatus(static_cast<IUnknown*>(GDynamicRHI->RHIGetNativeDevice()), &SleepStatusParams);
-
-		if (status == NVAPI_OK)
-		{
-			bFeatureSupport = true;
+			if (DriverVersion >= 51123)
+			{
+				bFlashIndicatorDriverControlled = true;
+			}
 		}
 	}
 }
 
 void FReflexLatencyMarkers::Tick(float DeltaTime)
 {
-	if (DisableLatencyMarkers == 0 && bEnabled && bProperDriverVersion && IsRHIDeviceNVIDIA() && bFeatureSupport)
+	if (DisableLatencyMarkers == 0 && bEnabled && bFeatureSupport)
 	{
 		NvAPI_Status LatencyStatus = NVAPI_OK;
 		NV_LATENCY_RESULT_PARAMS_V1 LatencyResults = { 0 };
@@ -120,7 +105,7 @@ void FReflexLatencyMarkers::Tick(float DeltaTime)
 
 void FReflexLatencyMarkers::SetInputSampleLatencyMarker(uint64 FrameNumber)
 {
-	if (DisableLatencyMarkers == 0 && bProperDriverVersion && bEnabled && IsRHIDeviceNVIDIA() && bFeatureSupport)
+	if (DisableLatencyMarkers == 0 && bEnabled && bFeatureSupport)
 	{
 		NvAPI_Status status = NVAPI_OK;
 		NV_LATENCY_MARKER_PARAMS_V1 params = { 0 };
@@ -134,7 +119,7 @@ void FReflexLatencyMarkers::SetInputSampleLatencyMarker(uint64 FrameNumber)
 
 void FReflexLatencyMarkers::SetSimulationLatencyMarkerStart(uint64 FrameNumber)
 {
-	if (DisableLatencyMarkers == 0 && bProperDriverVersion && bEnabled && IsRHIDeviceNVIDIA() && bFeatureSupport)
+	if (DisableLatencyMarkers == 0 && bEnabled && bFeatureSupport)
 	{
 		NvAPI_Status status = NVAPI_OK;
 		NV_LATENCY_MARKER_PARAMS_V1 params = { 0 };
@@ -148,7 +133,7 @@ void FReflexLatencyMarkers::SetSimulationLatencyMarkerStart(uint64 FrameNumber)
 
 void FReflexLatencyMarkers::SetSimulationLatencyMarkerEnd(uint64 FrameNumber)
 {
-	if (DisableLatencyMarkers == 0 && bProperDriverVersion && bEnabled && IsRHIDeviceNVIDIA() && bFeatureSupport)
+	if (DisableLatencyMarkers == 0 && bEnabled && bFeatureSupport)
 	{
 		NvAPI_Status status = NVAPI_OK;
 		NV_LATENCY_MARKER_PARAMS_V1 params = { 0 };
@@ -162,7 +147,7 @@ void FReflexLatencyMarkers::SetSimulationLatencyMarkerEnd(uint64 FrameNumber)
 
 void FReflexLatencyMarkers::SetRenderSubmitLatencyMarkerStart(uint64 FrameNumber)
 {
-	if (DisableLatencyMarkers == 0 && bProperDriverVersion && bEnabled && IsRHIDeviceNVIDIA() && bFeatureSupport)
+	if (DisableLatencyMarkers == 0 && bEnabled && bFeatureSupport)
 	{
 		NvAPI_Status status = NVAPI_OK;
 		NV_LATENCY_MARKER_PARAMS_V1 params = { 0 };
@@ -176,7 +161,7 @@ void FReflexLatencyMarkers::SetRenderSubmitLatencyMarkerStart(uint64 FrameNumber
 
 void FReflexLatencyMarkers::SetRenderSubmitLatencyMarkerEnd(uint64 FrameNumber)
 {
-	if (DisableLatencyMarkers == 0 && bProperDriverVersion && bEnabled && IsRHIDeviceNVIDIA() && bFeatureSupport)
+	if (DisableLatencyMarkers == 0 && bEnabled && bFeatureSupport)
 	{
 		NvAPI_Status status = NVAPI_OK;
 		NV_LATENCY_MARKER_PARAMS_V1 params = { 0 };
@@ -190,7 +175,7 @@ void FReflexLatencyMarkers::SetRenderSubmitLatencyMarkerEnd(uint64 FrameNumber)
 
 void FReflexLatencyMarkers::SetPresentLatencyMarkerStart(uint64 FrameNumber)
 {
-	if (DisableLatencyMarkers == 0 && bProperDriverVersion && bEnabled && IsRHIDeviceNVIDIA() && bFeatureSupport)
+	if (DisableLatencyMarkers == 0 && bEnabled && bFeatureSupport)
 	{
 		NvAPI_Status status = NVAPI_OK;
 		NV_LATENCY_MARKER_PARAMS_V1 params = { 0 };
@@ -204,7 +189,7 @@ void FReflexLatencyMarkers::SetPresentLatencyMarkerStart(uint64 FrameNumber)
 
 void FReflexLatencyMarkers::SetPresentLatencyMarkerEnd(uint64 FrameNumber)
 {
-	if (DisableLatencyMarkers == 0 && bProperDriverVersion && bEnabled && IsRHIDeviceNVIDIA() && bFeatureSupport)
+	if (DisableLatencyMarkers == 0 && bEnabled && bFeatureSupport)
 	{
 		NvAPI_Status status = NVAPI_OK;
 		NV_LATENCY_MARKER_PARAMS_V1 params = { 0 };
@@ -218,7 +203,7 @@ void FReflexLatencyMarkers::SetPresentLatencyMarkerEnd(uint64 FrameNumber)
 
 void FReflexLatencyMarkers::SetFlashIndicatorLatencyMarker(uint64 FrameNumber)
 {
-	if (DisableLatencyMarkers == 0 && bProperDriverVersion && bEnabled && IsRHIDeviceNVIDIA() && bFeatureSupport)
+	if (DisableLatencyMarkers == 0 && bEnabled && bFeatureSupport)
 	{
 		if (GetFlashIndicatorEnabled())
 		{
@@ -235,7 +220,7 @@ void FReflexLatencyMarkers::SetFlashIndicatorLatencyMarker(uint64 FrameNumber)
 
 void FReflexLatencyMarkers::SetCustomLatencyMarker(uint32 MarkerId, uint64 FrameNumber)
 {
-	if (DisableLatencyMarkers == 0 && bProperDriverVersion && bEnabled && IsRHIDeviceNVIDIA() && bFeatureSupport)
+	if (DisableLatencyMarkers == 0 && bEnabled && bFeatureSupport)
 	{
 		NvAPI_Status status = NVAPI_OK;
 		NV_LATENCY_MARKER_PARAMS_V1 params = { 0 };
@@ -254,7 +239,7 @@ void FReflexLatencyMarkers::SetFlashIndicatorEnabled(bool bInEnabled)
 
 bool FReflexLatencyMarkers::GetFlashIndicatorEnabled()
 {
-	if (DisableLatencyMarkers == 1 || !bProperDriverVersion || !IsRHIDeviceNVIDIA() || !bFeatureSupport)
+	if (DisableLatencyMarkers == 1 || !bFeatureSupport)
 	{
 		return false;
 	}
@@ -292,7 +277,7 @@ void FReflexLatencyMarkers::SetEnabled(bool bInEnabled)
 
 bool FReflexLatencyMarkers::GetEnabled()
 {
-	if (DisableLatencyMarkers == 1 || !bProperDriverVersion || !IsRHIDeviceNVIDIA() || !bFeatureSupport)
+	if (DisableLatencyMarkers == 1 || !bFeatureSupport)
 	{
 		return false;
 	}
@@ -302,7 +287,7 @@ bool FReflexLatencyMarkers::GetEnabled()
 
 bool FReflexLatencyMarkers::GetAvailable()
 {
-	if (DisableLatencyMarkers == 1 || !bProperDriverVersion || !IsRHIDeviceNVIDIA() || !bFeatureSupport)
+	if (DisableLatencyMarkers == 1 || !bFeatureSupport)
 	{
 		return false;
 	}
@@ -346,4 +331,4 @@ bool FReflexLatencyMarkers::Exec_Runtime(UWorld* InWorld, const TCHAR* Cmd, FOut
 	}
 
 	return bHandled;
-}
+}
\ No newline at end of file
diff --git a/Engine/Plugins/Runtime/Nvidia/Reflex/Source/Public/ReflexLatencyMarkers.h b/Engine/Plugins/Runtime/Nvidia/Reflex/Source/Public/ReflexLatencyMarkers.h
--- a/Engine/Plugins/Runtime/Nvidia/Reflex/Source/Public/ReflexLatencyMarkers.h
+++ b/Engine/Plugins/Runtime/Nvidia/Reflex/Source/Public/ReflexLatencyMarkers.h
@@ -12,7 +12,6 @@ class FReflexLatencyMarkers : public ILatencyMarkerModule, public FTickableGameO
 public:
 	virtual ~FReflexLatencyMarkers() {}
 
-	bool bProperDriverVersion = false;
 	bool bFeatureSupport = false;
 
 	float AverageTotalLatencyMs = 0.0f;
