Do not allow calls to NvAPI_D3D_SetLatencyMarker and NvAPI_D3D_GetLatency in the Reflex plugin on non-Nvidia adapters.

diff --git a/Engine/Plugins/Runtime/Nvidia/Reflex/Source/Private/ReflexLatencyMarkers.cpp b/Engine/Plugins/Runtime/Nvidia/Reflex/Source/Private/ReflexLatencyMarkers.cpp
--- a/Engine/Plugins/Runtime/Nvidia/Reflex/Source/Private/ReflexLatencyMarkers.cpp
+++ b/Engine/Plugins/Runtime/Nvidia/Reflex/Source/Private/ReflexLatencyMarkers.cpp
@@ -29,24 +29,23 @@ static FAutoConsoleVariableRef CVarDisableLatencyMarkers(
 
 void FReflexLatencyMarkers::Initialize()
 {
-	if ((RHIGetInterfaceType() == ERHIInterfaceType::D3D11 || RHIGetInterfaceType() == ERHIInterfaceType::D3D12))
+	bFeatureSupport = false;
+
+	if ( (IsRHIDeviceNVIDIA()) && (RHIGetInterfaceType() == ERHIInterfaceType::D3D11 || RHIGetInterfaceType() == ERHIInterfaceType::D3D12))
 	{
-		// No vendor or driver version required for basic Latency Marker support.
 		bFeatureSupport = true;
-		bFlashIndicatorDriverControlled = false;
-
+		
 		// Check for flash indicator sub-feature support, which still requires
 		// a minimum NVIDIA driver version.
-		if (IsRHIDeviceNVIDIA())
-		{
-			NvU32 DriverVersion = 0;
-			NvAPI_ShortString BranchString;
-			NvAPI_SYS_GetDriverAndBranchVersion(&DriverVersion, BranchString);
+		bFlashIndicatorDriverControlled = false;
 
-			if (DriverVersion >= 51123)
-			{
-				bFlashIndicatorDriverControlled = true;
-			}
+		NvU32 DriverVersion = 0;
+		NvAPI_ShortString BranchString;
+		NvAPI_SYS_GetDriverAndBranchVersion(&DriverVersion, BranchString);
+
+		if (DriverVersion >= 51123)
+		{
+			bFlashIndicatorDriverControlled = true;
 		}
 	}
 }
