Reflex: Stop forcing c++17

#rnx

diff --git a/Engine/Plugins/Runtime/Nvidia/Reflex/Source/Reflex.Build.cs b/Engine/Plugins/Runtime/Nvidia/Reflex/Source/Reflex.Build.cs
--- a/Engine/Plugins/Runtime/Nvidia/Reflex/Source/Reflex.Build.cs
+++ b/Engine/Plugins/Runtime/Nvidia/Reflex/Source/Reflex.Build.cs
@@ -8,29 +8,20 @@ namespace UnrealBuildTool.Rules
 	{
 		public Reflex(ReadOnlyTargetRules Target) : base(Target)
 		{
-			PCHUsage = ModuleRules.PCHUsageMode.UseExplicitOrSharedPCHs;
-
-			// Reevaulaute if ThirdParty\reflexstats.h is updated to compile with c++20
-			// error C4002: too many arguments for function-like macro invocation '_TLG_FOR_imp0'
-			CppStandard = CppStandardVersion.Cpp17;
-
-			PublicDependencyModuleNames.AddRange(
-				new string[]
-				{
-					"Core",
-					"Engine",
-					"RHI",
-					"CoreUObject",
-					"SlateCore",
-					"Slate"
-				}
-			);
+			PublicDependencyModuleNames.AddRange([
+				"Core",
+				"Engine",
+				"RHI",
+				"CoreUObject",
+				"SlateCore",
+				"Slate"
+			]);
 
 			// Grab NVAPI
 			AddEngineThirdPartyPrivateStaticDependencies(Target, "NVAPI");
 
 			// Grab ReflexStat
-			PublicIncludePaths.Add(Path.Combine(ModuleDirectory, "ThirdParty"));
+			PrivateIncludePaths.Add(Path.Combine(ModuleDirectory, "ThirdParty"));
 		}
 	}
 }
