Use FScopeIdle with NvAPI_D3D_Sleep for NVIDIA Reflex to remove time pollution from the game thread frame timing when Reflex is enabled in user settings

diff --git a/Engine/Plugins/Runtime/Nvidia/Reflex/Source/Private/ReflexMaxTickRateHandler.cpp b/Engine/Plugins/Runtime/Nvidia/Reflex/Source/Private/ReflexMaxTickRateHandler.cpp
--- a/Engine/Plugins/Runtime/Nvidia/Reflex/Source/Private/ReflexMaxTickRateHandler.cpp
+++ b/Engine/Plugins/Runtime/Nvidia/Reflex/Source/Private/ReflexMaxTickRateHandler.cpp
@@ -11,6 +11,7 @@
 #include "Engine/World.h"
 #include "GameFramework/PlayerController.h"
 #include "Kismet/GameplayStatics.h"
+#include "Stats/ThreadIdleStats.h"
 
 #include "Windows/AllowWindowsPlatformTypes.h"
 #include <D3D11.h>
@@ -108,7 +109,11 @@ bool FReflexMaxTickRateHandler::HandleMaxTickRate(float DesiredMaxTickRate)
 			}
 
 			NvAPI_Status StatusSleep = NVAPI_OK;
-			StatusSleep = NvAPI_D3D_Sleep(static_cast<IUnknown*>(GDynamicRHI->RHIGetNativeDevice()));
+			IUnknown* const Unknown = static_cast<IUnknown*>(GDynamicRHI->RHIGetNativeDevice());
+			{
+				UE::Stats::FThreadIdleStats::FScopeIdle Scope;
+				StatusSleep = NvAPI_D3D_Sleep(Unknown);
+			}
 
 			return true;
 		}
