Replace use of RHIGetNativeDevice with the various RHIGetDevice functions on platform RHI interfaces
- Using the platform RHI interface is better for type-safety, as there is no cast through void*.
- Also added RHIGetDevice_NoMGPU() for D3D12, to highlight areas of the codebase calling RHIGetDevice(0) that don't explicitly handle multi-GPU.

diff --git a/Engine/Plugins/Developer/RenderDocPlugin/Source/RenderDocPlugin/Private/RenderDocPluginModule.cpp b/Engine/Plugins/Developer/RenderDocPlugin/Source/RenderDocPlugin/Private/RenderDocPluginModule.cpp
--- a/Engine/Plugins/Developer/RenderDocPlugin/Source/RenderDocPlugin/Private/RenderDocPluginModule.cpp
+++ b/Engine/Plugins/Developer/RenderDocPlugin/Source/RenderDocPlugin/Private/RenderDocPluginModule.cpp
@@ -3,9 +3,11 @@
 #include "RenderDocPluginModule.h"
 #include "Framework/Application/SlateApplication.h"
 #include "UnrealClient.h"
+#include "IVulkanDynamicRHI.h"
 
 #if PLATFORM_WINDOWS
-#include "Windows/AllowWindowsPlatformTypes.h"
+#include "ID3D11DynamicRHI.h"
+#include "ID3D12DynamicRHI.h"
 #endif
 
 #include "Engine/Engine.h"
@@ -27,8 +29,6 @@
 #include "ISettingsSection.h"
 #include "LevelEditor.h"
 #include "SRenderDocPluginEditorExtension.h"
-#include "UnrealClient.h"
-extern UNREALED_API UEditorEngine* GEditor;
 #endif // WITH_EDITOR
 
 DEFINE_LOG_CATEGORY(RenderDocPlugin);
@@ -100,19 +100,25 @@ class FRenderDocFrameCapturer
 public:
 	static RENDERDOC_DevicePointer GetRenderdocDevicePointer()
 	{
-		RENDERDOC_DevicePointer Device;
-		if (RHIGetInterfaceType() == ERHIInterfaceType::Vulkan)
-		{
-			Device = GDynamicRHI->RHIGetNativeInstance();
 #ifndef RENDERDOC_DEVICEPOINTER_FROM_VKINSTANCE
 #define RENDERDOC_DEVICEPOINTER_FROM_VKINSTANCE(inst) (*((void **)(inst)))
 #endif
-			Device = RENDERDOC_DEVICEPOINTER_FROM_VKINSTANCE(Device);
-		}
-		else
+
+		RENDERDOC_DevicePointer Device {};
+
+		switch (RHIGetInterfaceType())
 		{
-			Device = GDynamicRHI->RHIGetNativeDevice();
+	#if PLATFORM_WINDOWS
+		case ERHIInterfaceType::D3D12: Device = GetID3D12DynamicRHI()->RHIGetDevice_NoMGPU(); break;
+		case ERHIInterfaceType::D3D11: Device = GetID3D11DynamicRHI()->RHIGetDevice(); break;
+	#endif
+
+		case ERHIInterfaceType::Vulkan:
+			Device = GetIVulkanDynamicRHI()->RHIGetVkInstance();
+			Device = RENDERDOC_DEVICEPOINTER_FROM_VKINSTANCE(Device);
+			break;
 		}
+
 		return Device;
 
 	}
@@ -717,8 +723,4 @@ void FRenderDocPluginModule::EndCapture_RenderThread(FRHICommandListImmediate* I
 
 #undef LOCTEXT_NAMESPACE
 
-#if PLATFORM_WINDOWS
-#include "Windows/HideWindowsPlatformTypes.h"
-#endif
-
 IMPLEMENT_MODULE(FRenderDocPluginModule, RenderDocPlugin)
diff --git a/Engine/Plugins/Developer/RenderDocPlugin/Source/RenderDocPlugin/RenderDocPlugin.Build.cs b/Engine/Plugins/Developer/RenderDocPlugin/Source/RenderDocPlugin/RenderDocPlugin.Build.cs
--- a/Engine/Plugins/Developer/RenderDocPlugin/Source/RenderDocPlugin/RenderDocPlugin.Build.cs
+++ b/Engine/Plugins/Developer/RenderDocPlugin/Source/RenderDocPlugin/RenderDocPlugin.Build.cs
@@ -8,7 +8,18 @@ namespace UnrealBuildTool.Rules
 	{
 		public RenderDocPlugin(ReadOnlyTargetRules Target) : base(Target)
         {
-            PrivateDependencyModuleNames.AddRange(new string[] { });            
+			PrivateIncludePathModuleNames.AddRange([
+				"VulkanRHI",
+			]);
+
+			if (Target.IsInPlatformGroup(UnrealPlatformGroup.Windows))
+			{
+				PrivateIncludePathModuleNames.AddRange([
+					"D3D11RHI",
+					"D3D12RHI"
+				]);
+			}
+
 			PublicDependencyModuleNames.AddRange(new string[]
 			{
 				"Core"
diff --git a/Engine/Plugins/Developer/XcodeGPUDebuggerPlugin/Source/XcodeGPUDebuggerPlugin/Private/XcodeGPUDebuggerPluginModule.cpp b/Engine/Plugins/Developer/XcodeGPUDebuggerPlugin/Source/XcodeGPUDebuggerPlugin/Private/XcodeGPUDebuggerPluginModule.cpp
--- a/Engine/Plugins/Developer/XcodeGPUDebuggerPlugin/Source/XcodeGPUDebuggerPlugin/Private/XcodeGPUDebuggerPluginModule.cpp
+++ b/Engine/Plugins/Developer/XcodeGPUDebuggerPlugin/Source/XcodeGPUDebuggerPlugin/Private/XcodeGPUDebuggerPluginModule.cpp
@@ -12,6 +12,7 @@
 #include "RenderingThread.h"
 #include "UnrealClient.h"
 #include "GeneralProjectSettings.h"
+#include "IMetalDynamicRHI.h"
 
 #if WITH_EDITOR
 #include "Editor.h"
@@ -61,7 +62,7 @@ public:
 		UE4_GEmitDrawEvents_BeforeCapture = GetEmitDrawEvents();
 		SetEmitDrawEvents(true);
 
-		MTL::Device* MetalDevice = (MTL::Device*)GDynamicRHI->RHIGetNativeDevice();
+		MTL::Device* MetalDevice = GetIMetalDynamicRHI()->RHIGetDevice();
 		MTL::CaptureManager* CaptureManager = MTL::CaptureManager::sharedCaptureManager();
 
 		MTL::CaptureDescriptor* CaptureDescriptor = MTL::CaptureDescriptor::alloc()->init();
@@ -349,7 +350,7 @@ void FXcodeGPUDebuggerPluginModule::BeginCapture(FRHICommandListImmediate* InRHI
 		CaptureFileName = FDateTime::Now().ToString();
 	}
 	
-    MTL::Device* MetalDevice = (MTL::Device*)GDynamicRHI->RHIGetNativeDevice();
+	MTL::Device* MetalDevice = GetIMetalDynamicRHI()->RHIGetDevice();
 
 	InRHICommandList->SubmitCommandsAndFlushGPU();
 	InRHICommandList->EnqueueLambda([FileName = CaptureFileName, MetalDevice](FRHICommandListImmediate& RHICommandList)
diff --git a/Engine/Plugins/Experimental/AVCodecs/AVCodecsCore/Source/AVCodecsCoreRHI/Private/AVCodecsCoreRHIModule.cpp b/Engine/Plugins/Experimental/AVCodecs/AVCodecsCore/Source/AVCodecsCoreRHI/Private/AVCodecsCoreRHIModule.cpp
--- a/Engine/Plugins/Experimental/AVCodecs/AVCodecsCore/Source/AVCodecsCoreRHI/Private/AVCodecsCoreRHIModule.cpp
+++ b/Engine/Plugins/Experimental/AVCodecs/AVCodecsCore/Source/AVCodecsCoreRHI/Private/AVCodecsCoreRHIModule.cpp
@@ -25,7 +25,7 @@
 	#include "Video/Resources/Vulkan/VideoResourceVulkan.h"
 #endif
 #if AVCODECS_USE_METAL
-	#include "DynamicRHI.h"
+	#include "IMetalDynamicRHI.h"
 	#include "Video/Resources/Metal/VideoResourceMetal.h"
 #endif
 
@@ -82,7 +82,7 @@ public:
 				case ERHIInterfaceType::D3D12:
 					FAVDevice::GetHardwareDevice()->SetContext<FVideoContextD3D12>(
 						MakeShared<FVideoContextD3D12>(
-							GetID3D12DynamicRHI()->RHIGetDevice(0)));
+							GetID3D12DynamicRHI()->RHIGetDevice_NoMGPU()));
 				
 					break;
 #endif
@@ -90,7 +90,9 @@ public:
                 case ERHIInterfaceType::Metal:
                     FAVDevice::GetHardwareDevice()->SetContext<FVideoContextMetal>(
                         MakeShared<FVideoContextMetal>(
-                            static_cast<MTL::Device*>(GDynamicRHI->RHIGetNativeDevice())));
+							GetIMetalDynamicRHI()->RHIGetDevice()
+						)
+					);
                     break;
 #endif
 				default:
diff --git a/Engine/Plugins/Media/AvfMedia/Source/AvfMedia/Private/Player/AvfMediaVideoSampler.cpp b/Engine/Plugins/Media/AvfMedia/Source/AvfMedia/Private/Player/AvfMediaVideoSampler.cpp
--- a/Engine/Plugins/Media/AvfMedia/Source/AvfMedia/Private/Player/AvfMediaVideoSampler.cpp
+++ b/Engine/Plugins/Media/AvfMedia/Source/AvfMedia/Private/Player/AvfMediaVideoSampler.cpp
@@ -193,7 +193,7 @@ void FAvfMediaVideoSampler::ProcessFrame(CVPixelBufferRef Frame, FTimespan Sampl
 	{
 		if (!MetalTextureCache)
 		{
-            id<MTLDevice> Device = (__bridge id<MTLDevice>)GDynamicRHI->RHIGetNativeDevice();
+            id<MTLDevice> Device = (__bridge id<MTLDevice>)GetIMetalDynamicRHI()->RHIGetDevice();
 			check(Device);
 			
 			CVReturn Return = CVMetalTextureCacheCreate(kCFAllocatorDefault, nullptr, Device, nullptr, &MetalTextureCache);
diff --git a/Engine/Plugins/Media/AvfMedia/Source/AvfMediaCapture/Private/Player/AvfMediaCapturePlayer.cpp b/Engine/Plugins/Media/AvfMedia/Source/AvfMediaCapture/Private/Player/AvfMediaCapturePlayer.cpp
--- a/Engine/Plugins/Media/AvfMedia/Source/AvfMediaCapture/Private/Player/AvfMediaCapturePlayer.cpp
+++ b/Engine/Plugins/Media/AvfMedia/Source/AvfMediaCapture/Private/Player/AvfMediaCapturePlayer.cpp
@@ -724,7 +724,7 @@ void FAvfMediaCapturePlayer::ProcessSampleBufferVideo(CMSampleBufferRef SampleBu
 		
 		if (!MetalTextureCache)
 		{
-            id<MTLDevice> Device = (__bridge id<MTLDevice>)GDynamicRHI->RHIGetNativeDevice();
+            id<MTLDevice> Device = (__bridge id<MTLDevice>)GetIMetalDynamicRHI()->RHIGetDevice();
 			check(Device);
 			
 			CVReturn Return = CVMetalTextureCacheCreate(kCFAllocatorDefault, nullptr, Device, nullptr, &MetalTextureCache);
diff --git a/Engine/Plugins/Media/D3D12VideoDecodersElectra/Source/D3D12VideoDecodersElectra.Build.cs b/Engine/Plugins/Media/D3D12VideoDecodersElectra/Source/D3D12VideoDecodersElectra.Build.cs
--- a/Engine/Plugins/Media/D3D12VideoDecodersElectra/Source/D3D12VideoDecodersElectra.Build.cs
+++ b/Engine/Plugins/Media/D3D12VideoDecodersElectra/Source/D3D12VideoDecodersElectra.Build.cs
@@ -32,6 +32,8 @@ namespace UnrealBuildTool.Rules
 					"ElectraCodecFactory",
 					"ElectraDecoders"
                 });
+
+			PrivateIncludePathModuleNames.Add("D3D12RHI");
 		}
 	}
 }
diff --git a/Engine/Plugins/Media/D3D12VideoDecodersElectra/Source/Private/VideoDecoder_D3D12.cpp b/Engine/Plugins/Media/D3D12VideoDecodersElectra/Source/Private/VideoDecoder_D3D12.cpp
--- a/Engine/Plugins/Media/D3D12VideoDecodersElectra/Source/Private/VideoDecoder_D3D12.cpp
+++ b/Engine/Plugins/Media/D3D12VideoDecodersElectra/Source/Private/VideoDecoder_D3D12.cpp
@@ -6,6 +6,7 @@
 #include "Features/IModularFeatures.h"
 #include "Features/IModularFeature.h"
 #include "RHI.h"
+#include "ID3D12DynamicRHI.h"
 
 #include "HAL/IConsoleManager.h"
 
@@ -66,7 +67,7 @@ int32 FCodecFormatHelper::FindSupportedFormats(ID3D12Device* InD3D12Device)
 		{
 			return 0;
 		}
-		DxDevice = static_cast<ID3D12Device*>(GDynamicRHI->RHIGetNativeDevice());
+		DxDevice = GetID3D12DynamicRHI()->RHIGetDevice_NoMGPU();
 	}
 	if (!DxDevice)
 	{
diff --git a/Engine/Plugins/Media/ElectraCodecs/Source/ElectraDecoders/ElectraDecoders.Build.cs b/Engine/Plugins/Media/ElectraCodecs/Source/ElectraDecoders/ElectraDecoders.Build.cs
--- a/Engine/Plugins/Media/ElectraCodecs/Source/ElectraDecoders/ElectraDecoders.Build.cs
+++ b/Engine/Plugins/Media/ElectraCodecs/Source/ElectraDecoders/ElectraDecoders.Build.cs
@@ -115,6 +115,8 @@ namespace UnrealBuildTool.Rules
 								"QuartzCore"
 				});
 
+				PrivateIncludePathModuleNames.Add("MetalRHI");
+
 				PrivateDefinitions.Add("ELECTRA_DECODERS_ENABLE_APPLE=1");
 				PrivateIncludePaths.Add("ElectraDecoders/Private/Apple");
 			}
diff --git a/Engine/Plugins/Media/ElectraCodecs/Source/ElectraDecoders/Private/Apple/AppleElectraDecodersPlatformResources.cpp b/Engine/Plugins/Media/ElectraCodecs/Source/ElectraDecoders/Private/Apple/AppleElectraDecodersPlatformResources.cpp
--- a/Engine/Plugins/Media/ElectraCodecs/Source/ElectraDecoders/Private/Apple/AppleElectraDecodersPlatformResources.cpp
+++ b/Engine/Plugins/Media/ElectraCodecs/Source/ElectraDecoders/Private/Apple/AppleElectraDecodersPlatformResources.cpp
@@ -13,6 +13,8 @@
 
 #if PLATFORM_MAC || PLATFORM_IOS || PLATFORM_TVOS
 
+#include "IMetalDynamicRHI.h"
+
 bool FElectraDecodersPlatformResourcesApple::GetMetalDevice(void** OutMetalDevice)
 {
     *OutMetalDevice = nullptr;
@@ -25,7 +27,7 @@ bool FElectraDecodersPlatformResourcesApple::GetMetalDevice(void** OutMetalDevic
 	{
 		return false;
 	}
-	*OutMetalDevice = GDynamicRHI->RHIGetNativeDevice();
+	*OutMetalDevice = GetIMetalDynamicRHI()->RHIGetDevice();
 	return true;
 }
 
diff --git a/Engine/Plugins/Media/ElectraCodecs/Source/ElectraDecoders/Private/DX/VideoDecoderH264_DX.cpp b/Engine/Plugins/Media/ElectraCodecs/Source/ElectraDecoders/Private/DX/VideoDecoderH264_DX.cpp
--- a/Engine/Plugins/Media/ElectraCodecs/Source/ElectraDecoders/Private/DX/VideoDecoderH264_DX.cpp
+++ b/Engine/Plugins/Media/ElectraCodecs/Source/ElectraDecoders/Private/DX/VideoDecoderH264_DX.cpp
@@ -976,7 +976,7 @@ void FElectraVideoDecoderH264_DX::CopyBufferData(TRefCountPtr<IMF2DBuffer> Buffe
 	check(SUCCEEDED(Res));
 
 	ID3D12DynamicRHI* const D3D12RHI = GetID3D12DynamicRHI();
-	ID3D12Device* const Device = static_cast<ID3D12Device*>(D3D12RHI->RHIGetNativeDevice());
+	ID3D12Device* const Device = D3D12RHI->RHIGetDevice_NoMGPU();
 
 	void* BufferAddr;
 	Res = BufferResource->Map(0, nullptr, &BufferAddr);
diff --git a/Engine/Plugins/Media/ElectraCodecs/Source/ElectraDecoders/Private/DX/VideoDecoderH265_DX.cpp b/Engine/Plugins/Media/ElectraCodecs/Source/ElectraDecoders/Private/DX/VideoDecoderH265_DX.cpp
--- a/Engine/Plugins/Media/ElectraCodecs/Source/ElectraDecoders/Private/DX/VideoDecoderH265_DX.cpp
+++ b/Engine/Plugins/Media/ElectraCodecs/Source/ElectraDecoders/Private/DX/VideoDecoderH265_DX.cpp
@@ -961,7 +961,7 @@ void FElectraVideoDecoderH265_DX::CopyBufferData(TRefCountPtr<IMF2DBuffer> Buffe
 	check(SUCCEEDED(Res));
 
 	ID3D12DynamicRHI* const D3D12RHI = GetID3D12DynamicRHI();
-	ID3D12Device* const Device = static_cast<ID3D12Device*>(D3D12RHI->RHIGetNativeDevice());
+	ID3D12Device* const Device = D3D12RHI->RHIGetDevice_NoMGPU();
 
 	void* BufferAddr;
 	Res = BufferResource->Map(0, nullptr, &BufferAddr);
diff --git a/Engine/Plugins/Media/ElectraCodecs/Source/ElectraDecoders/Private/Windows/WindowsElectraDecodersPlatformResources.cpp b/Engine/Plugins/Media/ElectraCodecs/Source/ElectraDecoders/Private/Windows/WindowsElectraDecodersPlatformResources.cpp
--- a/Engine/Plugins/Media/ElectraCodecs/Source/ElectraDecoders/Private/Windows/WindowsElectraDecodersPlatformResources.cpp
+++ b/Engine/Plugins/Media/ElectraCodecs/Source/ElectraDecoders/Private/Windows/WindowsElectraDecodersPlatformResources.cpp
@@ -3,6 +3,7 @@
 #include "WindowsElectraDecodersPlatformResources.h"
 #include "ElectraDecodersModule.h"
 #include "WindowsPlatformHeaders_Video_DX.h"
+#include "ID3D11DynamicRHI.h"
 #include "ID3D12DynamicRHI.h"
 #include "Async/Async.h"
 
@@ -62,14 +63,22 @@ bool FElectraDecodersPlatformResourcesWindows::GetD3DDeviceAndVersion(void** Out
 	{
 		return false;
 	}
-	auto RHIType = RHIGetInterfaceType();
-	if (RHIType != ERHIInterfaceType::D3D11 && RHIType != ERHIInterfaceType::D3D12)
+
+	switch (RHIGetInterfaceType())
 	{
+	case ERHIInterfaceType::D3D11:
+		*OutD3DDevice = GetID3D11DynamicRHI()->RHIGetDevice();
+		*OutD3DVersion = 11000;
+		return true;
+
+	case ERHIInterfaceType::D3D12:
+		*OutD3DDevice = GetID3D12DynamicRHI()->RHIGetDevice_NoMGPU();
+		*OutD3DVersion = 12000;
+		return true;
+
+	default:
 		return false;
 	}
-	*OutD3DDevice = GDynamicRHI->RHIGetNativeDevice();
-	*OutD3DVersion = RHIType == ERHIInterfaceType::D3D11 ? 11000 : 12000;
-	return true;
 }
 
 
diff --git a/Engine/Plugins/Media/ElectraUtil/Source/ElectraSamples/Private/Apple/ElectraTextureSample_Apple.cpp b/Engine/Plugins/Media/ElectraUtil/Source/ElectraSamples/Private/Apple/ElectraTextureSample_Apple.cpp
--- a/Engine/Plugins/Media/ElectraUtil/Source/ElectraSamples/Private/Apple/ElectraTextureSample_Apple.cpp
+++ b/Engine/Plugins/Media/ElectraUtil/Source/ElectraSamples/Private/Apple/ElectraTextureSample_Apple.cpp
@@ -66,7 +66,7 @@ void FElectraMediaTexConvApple::ConvertTexture(FTextureRHIRef& InDstTexture, CVI
 	{
 		if (!MetalTextureCache)
 		{
-            id<MTLDevice> Device = (__bridge id<MTLDevice>)GDynamicRHI->RHIGetNativeDevice();
+            id<MTLDevice> Device = (__bridge id<MTLDevice>)GetIMetalDynamicRHI()->RHIGetDevice();
 			check(Device);
 
 			CVReturn Return = CVMetalTextureCacheCreate(kCFAllocatorDefault, nullptr, Device, nullptr, &MetalTextureCache);
diff --git a/Engine/Plugins/Media/ElectraUtil/Source/ElectraSamples/Private/Windows/ElectraTextureSample_Windows.cpp b/Engine/Plugins/Media/ElectraUtil/Source/ElectraSamples/Private/Windows/ElectraTextureSample_Windows.cpp
--- a/Engine/Plugins/Media/ElectraUtil/Source/ElectraSamples/Private/Windows/ElectraTextureSample_Windows.cpp
+++ b/Engine/Plugins/Media/ElectraUtil/Source/ElectraSamples/Private/Windows/ElectraTextureSample_Windows.cpp
@@ -397,7 +397,7 @@ FRHITexture* FElectraTextureSample::GetTexture() const
 						Result = DxgiResource->GetSharedHandle(&SharedHandle);
 						if (SUCCEEDED(Result))
 						{
-							Result = GetID3D12DynamicRHI()->RHIGetDevice(0)->OpenSharedHandle(SharedHandle, __uuidof(ID3D12Resource), (void**)TextureCommonDX12.GetInitReference());
+							Result = GetID3D12DynamicRHI()->RHIGetDevice_NoMGPU()->OpenSharedHandle(SharedHandle, __uuidof(ID3D12Resource), (void**)TextureCommonDX12.GetInitReference());
 						}
 					}
 				}
@@ -452,7 +452,7 @@ struct FRHICommandCopyResourceDX11 final : public FRHICommand<FRHICommandCopyRes
 	void Execute(FRHICommandListBase& CmdList)
 	{
 		LLM_SCOPE(ELLMTag::MediaStreaming);
-		ID3D11Device* D3D11Device = static_cast<ID3D11Device*>(GDynamicRHI->RHIGetNativeDevice());
+		ID3D11Device* D3D11Device = GetID3D11DynamicRHI()->RHIGetDevice();
 		ID3D11DeviceContext* D3D11DeviceContext = nullptr;
 
 		D3D11Device->GetImmediateContext(&D3D11DeviceContext);
diff --git a/Engine/Plugins/Media/MediaIOFramework/Source/GPUTextureTransfer/Private/D3D12TextureTransfer.cpp b/Engine/Plugins/Media/MediaIOFramework/Source/GPUTextureTransfer/Private/D3D12TextureTransfer.cpp
--- a/Engine/Plugins/Media/MediaIOFramework/Source/GPUTextureTransfer/Private/D3D12TextureTransfer.cpp
+++ b/Engine/Plugins/Media/MediaIOFramework/Source/GPUTextureTransfer/Private/D3D12TextureTransfer.cpp
@@ -15,7 +15,7 @@ namespace UE::GPUTextureTransfer::Private
 	{
 		ID3D12DynamicRHI* D3D12RHI = GetID3D12DynamicRHI();
 
-		D3DDevice = D3D12RHI->RHIGetDevice(0);
+		D3DDevice = D3D12RHI->RHIGetDevice_NoMGPU();
 		D3DCommandQueue = D3D12RHI->RHIGetCommandQueue();
 
 		return dvpInitD3D12Device(D3DDevice, 0);
diff --git a/Engine/Plugins/Media/WmfMedia/Source/WmfMedia/Private/Wmf/WmfMediaSink.cpp b/Engine/Plugins/Media/WmfMedia/Source/WmfMedia/Private/Wmf/WmfMediaSink.cpp
--- a/Engine/Plugins/Media/WmfMedia/Source/WmfMedia/Private/Wmf/WmfMediaSink.cpp
+++ b/Engine/Plugins/Media/WmfMedia/Source/WmfMedia/Private/Wmf/WmfMediaSink.cpp
@@ -19,6 +19,9 @@
 #include <d3d11.h>
 #include <d3d12.h>
 #include <dxgi1_4.h>
+#include "Windows/HideWindowsPlatformTypes.h"
+
+#include "ID3D11DynamicRHI.h"
 
 /* FWmfMediaStreamSink structors
  *****************************************************************************/
@@ -513,7 +516,7 @@ bool FWmfMediaSink::CreateDXGIManagerAndDevice()
 				return false;
 			}
 
-			ID3D11Device* PreExistingD3D11Device = static_cast<ID3D11Device*>(GDynamicRHI->RHIGetNativeDevice());
+			ID3D11Device* PreExistingD3D11Device = GetID3D11DynamicRHI()->RHIGetDevice();
 
 			TComPtr<IDXGIDevice> DXGIDevice;
 			PreExistingD3D11Device->QueryInterface(__uuidof(IDXGIDevice), (void**)&DXGIDevice);
@@ -575,7 +578,7 @@ bool FWmfMediaSink::CreateDXGIManagerAndDevice()
 		D3D11Device->QueryInterface(__uuidof(ID3D10Multithread), (void**)&MultiThread);
 		if (MultiThread)
 		{
-			MultiThread->SetMultithreadProtected(TRUE);
+			MultiThread->SetMultithreadProtected(true);
 			MultiThread->Release();
 		}
 		else
@@ -593,7 +596,5 @@ bool FWmfMediaSink::CreateDXGIManagerAndDevice()
 	}
 }
 
-#include "Windows/HideWindowsPlatformTypes.h"
-
 #endif
 
diff --git a/Engine/Plugins/Runtime/AR/AppleAR/AppleARKit/Source/AppleARKit/Private/AppleARKitSystem.cpp b/Engine/Plugins/Runtime/AR/AppleAR/AppleARKit/Source/AppleARKit/Private/AppleARKitSystem.cpp
--- a/Engine/Plugins/Runtime/AR/AppleAR/AppleARKit/Source/AppleARKit/Private/AppleARKitSystem.cpp
+++ b/Engine/Plugins/Runtime/AR/AppleAR/AppleARKit/Source/AppleARKit/Private/AppleARKitSystem.cpp
@@ -27,6 +27,10 @@
 #include "RenderGraphBuilder.h"
 #include "DataDrivenShaderPlatformInfo.h"
 
+#if PLATFORM_APPLE
+#include "IMetalDynamicRHI.h"
+#endif
+
 // For mesh occlusion
 #include "MRMeshComponent.h"
 #include "AROriginActor.h"
@@ -1653,7 +1657,7 @@ bool FAppleARKitSystem::Run(UARSessionConfig* SessionConfig)
 		// Create MetalTextureCache
 		if (IsMetalPlatform(GMaxRHIShaderPlatform))
 		{
-            id<MTLDevice> Device = (__bridge id<MTLDevice>)GDynamicRHI->RHIGetNativeDevice();
+            id<MTLDevice> Device = (__bridge id<MTLDevice>)GetIMetalDynamicRHI()->RHIGetDevice();
 			check(Device);
 
 			CVReturn Return = CVMetalTextureCacheCreate(nullptr, nullptr, (__bridge id<MTLDevice>)Device, nullptr, &MetalTextureCache);
diff --git a/Engine/Plugins/Runtime/AR/AppleAR/AppleARKit/Source/AppleARKit/Private/AppleARKitTextures.cpp b/Engine/Plugins/Runtime/AR/AppleAR/AppleARKit/Source/AppleARKit/Private/AppleARKitTextures.cpp
--- a/Engine/Plugins/Runtime/AR/AppleAR/AppleARKit/Source/AppleARKit/Private/AppleARKitTextures.cpp
+++ b/Engine/Plugins/Runtime/AR/AppleAR/AppleARKit/Source/AppleARKit/Private/AppleARKitTextures.cpp
@@ -273,7 +273,7 @@ protected:
 		{
 			if (!ImageFilter)
 			{
-				id<MTLDevice> MetalDevice = (__bridge id<MTLDevice>)GDynamicRHI->RHIGetNativeDevice();
+				id<MTLDevice> MetalDevice = (__bridge id<MTLDevice>)GetIMetalDynamicRHI()->RHIGetDevice();
 				ImageFilter = MakeShared<FAppleImageFilter, ESPMode::ThreadSafe>(MetalDevice);
 			}
 			
diff --git a/Engine/Plugins/Runtime/AR/AppleAR/AppleARKit/Source/AppleARKit/Private/AppleARKitVideoOverlay.cpp b/Engine/Plugins/Runtime/AR/AppleAR/AppleARKit/Source/AppleARKit/Private/AppleARKitVideoOverlay.cpp
--- a/Engine/Plugins/Runtime/AR/AppleAR/AppleARKit/Source/AppleARKit/Private/AppleARKitVideoOverlay.cpp
+++ b/Engine/Plugins/Runtime/AR/AppleAR/AppleARKit/Source/AppleARKit/Private/AppleARKitVideoOverlay.cpp
@@ -244,7 +244,7 @@ void FAppleARKitVideoOverlay::UpdateOcclusionTextures(const FAppleARKitFrame& Fr
 		if (OcclusionType == EARKitOcclusionType::PersonSegmentation && NativeFrame &&
 			(NativeFrame.segmentationBuffer || NativeFrame.estimatedDepthData))
 		{
-            id<MTLDevice> Device = (__bridge id<MTLDevice>)GDynamicRHI->RHIGetNativeDevice();
+            id<MTLDevice> Device = (__bridge id<MTLDevice>)GetIMetalDynamicRHI()->RHIGetDevice();
 			check(Device);
 
 			if (!MatteGenerator)
diff --git a/Engine/Plugins/Runtime/AppleMoviePlayer/Source/AppleMoviePlayer/Private/AppleMovieStreamer.cpp b/Engine/Plugins/Runtime/AppleMoviePlayer/Source/AppleMoviePlayer/Private/AppleMovieStreamer.cpp
--- a/Engine/Plugins/Runtime/AppleMoviePlayer/Source/AppleMoviePlayer/Private/AppleMovieStreamer.cpp
+++ b/Engine/Plugins/Runtime/AppleMoviePlayer/Source/AppleMoviePlayer/Private/AppleMovieStreamer.cpp
@@ -131,7 +131,7 @@ bool FAVPlayerMovieStreamer::Init(const TArray<FString>& MoviePaths, TEnumAsByte
 {
 	if(MetalTextureCache == NULL && FApp::CanEverRender())
 	{
-        id<MTLDevice> Device = (__bridge id<MTLDevice>)GDynamicRHI->RHIGetNativeDevice();
+        id<MTLDevice> Device = (__bridge id<MTLDevice>)GetIMetalDynamicRHI()->RHIGetDevice();
 		check(Device != nil);
 
 		CVReturn Return = CVMetalTextureCacheCreate(kCFAllocatorDefault, nullptr, Device, nullptr, &MetalTextureCache);
diff --git a/Engine/Plugins/Runtime/Nvidia/Reflex/Source/Private/ReflexLatencyMarkers.cpp b/Engine/Plugins/Runtime/Nvidia/Reflex/Source/Private/ReflexLatencyMarkers.cpp
--- a/Engine/Plugins/Runtime/Nvidia/Reflex/Source/Private/ReflexLatencyMarkers.cpp
+++ b/Engine/Plugins/Runtime/Nvidia/Reflex/Source/Private/ReflexLatencyMarkers.cpp
@@ -1,6 +1,7 @@
 // Copyright Epic Games, Inc. All Rights Reserved.
 
 #include "ReflexLatencyMarkers.h" 
+#include "ReflexModule.h"
 
 #include "Engine/GameViewportClient.h"
 #include "HAL/IConsoleManager.h"
@@ -58,7 +59,7 @@ void FReflexLatencyMarkers::Tick(float DeltaTime)
 		NV_LATENCY_RESULT_PARAMS_V1 LatencyResults = { 0 };
 		LatencyResults.version = NV_LATENCY_RESULT_PARAMS_VER1;
 
-		LatencyStatus = NvAPI_D3D_GetLatency(static_cast<IUnknown*>(GDynamicRHI->RHIGetNativeDevice()), &LatencyResults);
+		LatencyStatus = NvAPI_D3D_GetLatency(FReflexModule::GetD3DDevice(), &LatencyResults);
 
 		if (LatencyStatus == NVAPI_OK)
 		{
@@ -112,7 +113,7 @@ void FReflexLatencyMarkers::SetInputSampleLatencyMarker(uint64 FrameNumber)
 		params.frameID = FrameNumber;
 		params.markerType = INPUT_SAMPLE;
 
-		status = NvAPI_D3D_SetLatencyMarker(static_cast<IUnknown*>(GDynamicRHI->RHIGetNativeDevice()), &params);
+		status = NvAPI_D3D_SetLatencyMarker(FReflexModule::GetD3DDevice(), &params);
 	}
 }
 
@@ -126,7 +127,7 @@ void FReflexLatencyMarkers::SetSimulationLatencyMarkerStart(uint64 FrameNumber)
 		params.frameID = FrameNumber;
 		params.markerType = SIMULATION_START;
 
-		status = NvAPI_D3D_SetLatencyMarker(static_cast<IUnknown*>(GDynamicRHI->RHIGetNativeDevice()), &params);
+		status = NvAPI_D3D_SetLatencyMarker(FReflexModule::GetD3DDevice(), &params);
 	}
 }
 
@@ -140,7 +141,7 @@ void FReflexLatencyMarkers::SetSimulationLatencyMarkerEnd(uint64 FrameNumber)
 		params.frameID = FrameNumber;
 		params.markerType = SIMULATION_END;
 
-		status = NvAPI_D3D_SetLatencyMarker(static_cast<IUnknown*>(GDynamicRHI->RHIGetNativeDevice()), &params);
+		status = NvAPI_D3D_SetLatencyMarker(FReflexModule::GetD3DDevice(), &params);
 	}
 }
 
@@ -154,7 +155,7 @@ void FReflexLatencyMarkers::SetRenderSubmitLatencyMarkerStart(uint64 FrameNumber
 		params.frameID = FrameNumber;
 		params.markerType = RENDERSUBMIT_START;
 
-		status = NvAPI_D3D_SetLatencyMarker(static_cast<IUnknown*>(GDynamicRHI->RHIGetNativeDevice()), &params);
+		status = NvAPI_D3D_SetLatencyMarker(FReflexModule::GetD3DDevice(), &params);
 	}
 }
 
@@ -168,7 +169,7 @@ void FReflexLatencyMarkers::SetRenderSubmitLatencyMarkerEnd(uint64 FrameNumber)
 		params.frameID = FrameNumber;
 		params.markerType = RENDERSUBMIT_END;
 
-		status = NvAPI_D3D_SetLatencyMarker(static_cast<IUnknown*>(GDynamicRHI->RHIGetNativeDevice()), &params);
+		status = NvAPI_D3D_SetLatencyMarker(FReflexModule::GetD3DDevice(), &params);
 	}
 }
 
@@ -182,7 +183,7 @@ void FReflexLatencyMarkers::SetPresentLatencyMarkerStart(uint64 FrameNumber)
 		params.frameID = FrameNumber;
 		params.markerType = PRESENT_START;
 
-		status = NvAPI_D3D_SetLatencyMarker(static_cast<IUnknown*>(GDynamicRHI->RHIGetNativeDevice()), &params);
+		status = NvAPI_D3D_SetLatencyMarker(FReflexModule::GetD3DDevice(), &params);
 	}
 }
 
@@ -196,7 +197,7 @@ void FReflexLatencyMarkers::SetPresentLatencyMarkerEnd(uint64 FrameNumber)
 		params.frameID = FrameNumber;
 		params.markerType = PRESENT_END;
 
-		status = NvAPI_D3D_SetLatencyMarker(static_cast<IUnknown*>(GDynamicRHI->RHIGetNativeDevice()), &params);
+		status = NvAPI_D3D_SetLatencyMarker(FReflexModule::GetD3DDevice(), &params);
 	}
 }
 
@@ -212,7 +213,7 @@ void FReflexLatencyMarkers::SetFlashIndicatorLatencyMarker(uint64 FrameNumber)
 			params.frameID = FrameNumber;
 			params.markerType = TRIGGER_FLASH;
 
-			status = NvAPI_D3D_SetLatencyMarker(static_cast<IUnknown*>(GDynamicRHI->RHIGetNativeDevice()), &params);
+			status = NvAPI_D3D_SetLatencyMarker(FReflexModule::GetD3DDevice(), &params);
 		}
 	}
 }
@@ -227,7 +228,7 @@ void FReflexLatencyMarkers::SetCustomLatencyMarker(uint32 MarkerId, uint64 Frame
 		params.frameID = FrameNumber;
 		params.markerType = NV_LATENCY_MARKER_TYPE(MarkerId);
 
-		status = NvAPI_D3D_SetLatencyMarker(static_cast<IUnknown*>(GDynamicRHI->RHIGetNativeDevice()), &params);
+		status = NvAPI_D3D_SetLatencyMarker(FReflexModule::GetD3DDevice(), &params);
 	}
 }
 
diff --git a/Engine/Plugins/Runtime/Nvidia/Reflex/Source/Private/ReflexMaxTickRateHandler.cpp b/Engine/Plugins/Runtime/Nvidia/Reflex/Source/Private/ReflexMaxTickRateHandler.cpp
--- a/Engine/Plugins/Runtime/Nvidia/Reflex/Source/Private/ReflexMaxTickRateHandler.cpp
+++ b/Engine/Plugins/Runtime/Nvidia/Reflex/Source/Private/ReflexMaxTickRateHandler.cpp
@@ -1,6 +1,7 @@
 // Copyright Epic Games, Inc. All Rights Reserved.
 
 #include "ReflexMaxTickRateHandler.h" 
+#include "ReflexModule.h"
 
 #include "Engine/GameViewportClient.h"
 #include "HAL/IConsoleManager.h"
@@ -67,7 +68,7 @@ void FReflexMaxTickRateHandler::Initialize()
 		SleepStatusParams.version = NV_GET_SLEEP_STATUS_PARAMS_VER1;
 
 		NvAPI_Status status = NVAPI_OK;
-		status = NvAPI_D3D_GetSleepStatus(static_cast<IUnknown*>(GDynamicRHI->RHIGetNativeDevice()), &SleepStatusParams);
+		status = NvAPI_D3D_GetSleepStatus(FReflexModule::GetD3DDevice(), &SleepStatusParams);
 
 		if (status == NVAPI_OK)
 		{
@@ -93,7 +94,7 @@ bool FReflexMaxTickRateHandler::HandleMaxTickRate(float DesiredMaxTickRate)
 				MinimumInterval = DesiredMinimumInterval;
 				params.minimumIntervalUs = MinimumInterval;
 				params.bUseMarkersToOptimize = DisableLatencyMarkerOptimize ? NV_FALSE : NV_TRUE;
-				status = NvAPI_D3D_SetSleepMode(static_cast<IUnknown*>(GDynamicRHI->RHIGetNativeDevice()), &params);
+				status = NvAPI_D3D_SetSleepMode(FReflexModule::GetD3DDevice(), &params);
 
 				// Need to verify that Low Latency mode change actually applied successfully
 				if (bLowLatencyMode && (status != NVAPI_OK))
@@ -109,7 +110,7 @@ bool FReflexMaxTickRateHandler::HandleMaxTickRate(float DesiredMaxTickRate)
 			}
 
 			NvAPI_Status StatusSleep = NVAPI_OK;
-			IUnknown* const Unknown = static_cast<IUnknown*>(GDynamicRHI->RHIGetNativeDevice());
+			IUnknown* const Unknown = FReflexModule::GetD3DDevice();
 			{
 				UE::Stats::FThreadIdleStats::FScopeIdle Scope;
 				StatusSleep = NvAPI_D3D_Sleep(Unknown);
@@ -129,7 +130,7 @@ bool FReflexMaxTickRateHandler::HandleMaxTickRate(float DesiredMaxTickRate)
 				params.bLowLatencyBoost = false;
 				params.minimumIntervalUs = 0;
 				params.bUseMarkersToOptimize = false;
-				status = NvAPI_D3D_SetSleepMode(static_cast<IUnknown*>(GDynamicRHI->RHIGetNativeDevice()), &params);
+				status = NvAPI_D3D_SetSleepMode(FReflexModule::GetD3DDevice(), &params);
 				UE_LOG(LogMaxTickRateHandler, Log, TEXT("SetSleepMode clean up"));
 
 				// Reset module back to default values in case re-enabled in the same session
diff --git a/Engine/Plugins/Runtime/Nvidia/Reflex/Source/Private/ReflexModule.cpp b/Engine/Plugins/Runtime/Nvidia/Reflex/Source/Private/ReflexModule.cpp
--- a/Engine/Plugins/Runtime/Nvidia/Reflex/Source/Private/ReflexModule.cpp
+++ b/Engine/Plugins/Runtime/Nvidia/Reflex/Source/Private/ReflexModule.cpp
@@ -4,6 +4,9 @@
 
 #include "Modules/ModuleManager.h"
 
+#include "ID3D11DynamicRHI.h"
+#include "ID3D12DynamicRHI.h"
+
 IMPLEMENT_MODULE(FReflexModule, Reflex);
 
 void FReflexModule::StartupModule()
@@ -46,4 +49,15 @@ void FReflexModule::ShutdownModule()
 	{
 		IModularFeatures::Get().UnregisterModularFeature(ReflexStatsLatencyMarker->GetModularFeatureName(), ReflexStatsLatencyMarker.Get());
 	}
-}
\ No newline at end of file
+}
+
+IUnknown* FReflexModule::GetD3DDevice()
+{
+	switch (RHIGetInterfaceType())
+	{
+	case ERHIInterfaceType::D3D11: return GetID3D11DynamicRHI()->RHIGetDevice();
+	case ERHIInterfaceType::D3D12: return GetID3D12DynamicRHI()->RHIGetDevice_NoMGPU();
+	}
+
+	return nullptr;
+}
diff --git a/Engine/Plugins/Runtime/Nvidia/Reflex/Source/Private/ReflexModule.h b/Engine/Plugins/Runtime/Nvidia/Reflex/Source/Private/ReflexModule.h
--- a/Engine/Plugins/Runtime/Nvidia/Reflex/Source/Private/ReflexModule.h
+++ b/Engine/Plugins/Runtime/Nvidia/Reflex/Source/Private/ReflexModule.h
@@ -18,6 +18,8 @@ public:
 	virtual void StartupModule() override;
 	virtual void ShutdownModule() override;
 	// ~End IModuleInterface
+	
+	static struct IUnknown* GetD3DDevice();
 
 private:
 
diff --git a/Engine/Plugins/Runtime/Nvidia/Reflex/Source/Reflex.Build.cs b/Engine/Plugins/Runtime/Nvidia/Reflex/Source/Reflex.Build.cs
--- a/Engine/Plugins/Runtime/Nvidia/Reflex/Source/Reflex.Build.cs
+++ b/Engine/Plugins/Runtime/Nvidia/Reflex/Source/Reflex.Build.cs
@@ -17,6 +17,11 @@ namespace UnrealBuildTool.Rules
 				"Slate"
 			]);
 
+			PrivateIncludePathModuleNames.AddRange([
+				"D3D11RHI",
+				"D3D12RHI"
+			]);
+
 			// Grab NVAPI
 			AddEngineThirdPartyPrivateStaticDependencies(Target, "NVAPI");
 
diff --git a/Engine/Plugins/Runtime/OpenXR/Source/OpenXRHMD/Private/OpenXRHMD_RenderBridge.cpp b/Engine/Plugins/Runtime/OpenXR/Source/OpenXRHMD/Private/OpenXRHMD_RenderBridge.cpp
--- a/Engine/Plugins/Runtime/OpenXR/Source/OpenXRHMD/Private/OpenXRHMD_RenderBridge.cpp
+++ b/Engine/Plugins/Runtime/OpenXR/Source/OpenXRHMD/Private/OpenXRHMD_RenderBridge.cpp
@@ -132,7 +132,7 @@ public:
 
 	virtual void* GetGraphicsBinding(XrSystemId InSystem) override
 	{
-		LUID AdapterLuid = GetID3D12DynamicRHI()->RHIGetDevice(0)->GetAdapterLuid();
+		LUID AdapterLuid = GetID3D12DynamicRHI()->RHIGetDevice_NoMGPU()->GetAdapterLuid();
 		if (reinterpret_cast<uint64&>(AdapterLuid) != GetGraphicsAdapterLuid(InSystem))
 		{
 			return nullptr;
@@ -140,7 +140,7 @@ public:
 
 		Binding.type = XR_TYPE_GRAPHICS_BINDING_D3D12_KHR;
 		Binding.next = nullptr;
-		Binding.device = GetID3D12DynamicRHI()->RHIGetDevice(0);
+		Binding.device = GetID3D12DynamicRHI()->RHIGetDevice_NoMGPU();
 		Binding.queue = GetID3D12DynamicRHI()->RHIGetCommandQueue();
 		return &Binding;
 	}
diff --git a/Engine/Plugins/Runtime/XR/XRScribe/Source/XRScribe/Private/XRScribeEmulationLayer.cpp b/Engine/Plugins/Runtime/XR/XRScribe/Source/XRScribe/Private/XRScribeEmulationLayer.cpp
--- a/Engine/Plugins/Runtime/XR/XRScribe/Source/XRScribe/Private/XRScribeEmulationLayer.cpp
+++ b/Engine/Plugins/Runtime/XR/XRScribe/Source/XRScribe/Private/XRScribeEmulationLayer.cpp
@@ -2657,7 +2657,7 @@ XrResult FOpenXREmulationLayer::XrLayerGetD3D12GraphicsRequirementsKHR(XrInstanc
 
 	if (GDynamicRHI != nullptr)
 	{
-		graphicsRequirements->adapterLuid = GetID3D12DynamicRHI()->RHIGetDevice(0)->GetAdapterLuid();
+		graphicsRequirements->adapterLuid = GetID3D12DynamicRHI()->RHIGetDevice_NoMGPU()->GetAdapterLuid();
 	}
 	else
 	{
diff --git a/Engine/Plugins/Runtime/nDisplay/Source/DisplayCluster/Private/Render/Synchronization/Windows/DisplayClusterRenderSyncPolicyNvidiaPresentBarrierWindows.cpp b/Engine/Plugins/Runtime/nDisplay/Source/DisplayCluster/Private/Render/Synchronization/Windows/DisplayClusterRenderSyncPolicyNvidiaPresentBarrierWindows.cpp
--- a/Engine/Plugins/Runtime/nDisplay/Source/DisplayCluster/Private/Render/Synchronization/Windows/DisplayClusterRenderSyncPolicyNvidiaPresentBarrierWindows.cpp
+++ b/Engine/Plugins/Runtime/nDisplay/Source/DisplayCluster/Private/Render/Synchronization/Windows/DisplayClusterRenderSyncPolicyNvidiaPresentBarrierWindows.cpp
@@ -22,6 +22,7 @@ THIRD_PARTY_INCLUDES_START
 THIRD_PARTY_INCLUDES_END
 #include "Windows/HideWindowsPlatformTypes.h"
 
+#include "ID3D12DynamicRHI.h"
 
 namespace DisplayClusterRenderSyncPolicyNvidiaPresentBarrier_Data_Windows
 {
@@ -139,7 +140,7 @@ bool FDisplayClusterRenderSyncPolicyNvidiaPresentBarrier::InitializePresentBarri
 	}
 
 	// Get D3D1XDevice
-	D3DDevice = static_cast<ID3D12Device*>(GDynamicRHI->RHIGetNativeDevice());
+	D3DDevice = GetID3D12DynamicRHI()->RHIGetDevice_NoMGPU();
 	check(D3DDevice);
 
 	// Get IDXGISwapChain
diff --git a/Engine/Plugins/Runtime/nDisplay/Source/DisplayCluster/Private/Render/Synchronization/Windows/DisplayClusterRenderSyncPolicyNvidiaSwapBarrierWindows.cpp b/Engine/Plugins/Runtime/nDisplay/Source/DisplayCluster/Private/Render/Synchronization/Windows/DisplayClusterRenderSyncPolicyNvidiaSwapBarrierWindows.cpp
--- a/Engine/Plugins/Runtime/nDisplay/Source/DisplayCluster/Private/Render/Synchronization/Windows/DisplayClusterRenderSyncPolicyNvidiaSwapBarrierWindows.cpp
+++ b/Engine/Plugins/Runtime/nDisplay/Source/DisplayCluster/Private/Render/Synchronization/Windows/DisplayClusterRenderSyncPolicyNvidiaSwapBarrierWindows.cpp
@@ -23,6 +23,8 @@ THIRD_PARTY_INCLUDES_START
 THIRD_PARTY_INCLUDES_END
 #include "Windows/HideWindowsPlatformTypes.h"
 
+#include "ID3D12DynamicRHI.h"
+#include "ID3D11DynamicRHI.h"
 
 namespace DisplayClusterRenderSyncPolicyNvidiaSwapBarrier_Data_Windows
 {
@@ -169,7 +171,12 @@ bool FDisplayClusterRenderSyncPolicyNvidiaSwapBarrier::InitializeNvidiaSwapLock(
 #if WITH_NVAPI
 
 	// Get D3D1XDevice
-	D3DDevice = static_cast<IUnknown*>(GDynamicRHI->RHIGetNativeDevice());
+	D3DDevice = nullptr;
+	switch (RHIGetInterfaceType())
+	{
+	case ERHIInterfaceType::D3D12: D3DDevice = GetID3D12DynamicRHI()->RHIGetDevice_NoMGPU(); break;
+	case ERHIInterfaceType::D3D11: D3DDevice = GetID3D11DynamicRHI()->RHIGetDevice(); break;
+	}
 	check(D3DDevice);
 	
 	// Get IDXGISwapChain
diff --git a/Engine/Plugins/Runtime/nDisplay/Source/DisplayClusterProjection/Private/Policy/EasyBlend/Windows/DX12/DisplayClusterProjectionEasyBlendPolicyViewDataDX12.cpp b/Engine/Plugins/Runtime/nDisplay/Source/DisplayClusterProjection/Private/Policy/EasyBlend/Windows/DX12/DisplayClusterProjectionEasyBlendPolicyViewDataDX12.cpp
--- a/Engine/Plugins/Runtime/nDisplay/Source/DisplayClusterProjection/Private/Policy/EasyBlend/Windows/DX12/DisplayClusterProjectionEasyBlendPolicyViewDataDX12.cpp
+++ b/Engine/Plugins/Runtime/nDisplay/Source/DisplayClusterProjection/Private/Policy/EasyBlend/Windows/DX12/DisplayClusterProjectionEasyBlendPolicyViewDataDX12.cpp
@@ -14,29 +14,6 @@
 #include "Misc/Paths.h"
 #include "UnrealClient.h"
 
-namespace UE::DisplayClusterProjection::EasyBlend
-{
-	static inline ID3D12Device2* GetD3D12Device2_RenderThread(FRHICommandList& RHICmdList)
-	{
-		// Always use device 0 for EasyBlend
-		const uint32 DevIndex = 0;
-
-		if (ID3D12Device* D3D12Device = GetID3D12DynamicRHI()->RHIGetDevice(DevIndex))
-		{
-			TRefCountPtr<ID3D12Device2> D3D12Device2 = nullptr;
-			HRESULT Result = D3D12Device->QueryInterface(IID_PPV_ARGS(D3D12Device2.GetInitReference()));
-			if (SUCCEEDED(Result))
-			{
-				return D3D12Device2;
-			}
-
-			UE_LOG(LogDisplayClusterProjectionEasyBlend, Error, TEXT("Couldn't initialize ID3D12Device2"));
-		}
-
-		return nullptr;
-	}
-};
-
 //------------------------------------------------------------------------------
 // FDisplayClusterProjectionEasyBlendPolicyViewDataDX12
 //------------------------------------------------------------------------------
diff --git a/Engine/Plugins/Runtime/nDisplay/Source/DisplayClusterProjection/Private/Policy/VIOSO/Windows/DisplayClusterProjectionVIOSOPolicyViewData.cpp b/Engine/Plugins/Runtime/nDisplay/Source/DisplayClusterProjection/Private/Policy/VIOSO/Windows/DisplayClusterProjectionVIOSOPolicyViewData.cpp
--- a/Engine/Plugins/Runtime/nDisplay/Source/DisplayClusterProjection/Private/Policy/VIOSO/Windows/DisplayClusterProjectionVIOSOPolicyViewData.cpp
+++ b/Engine/Plugins/Runtime/nDisplay/Source/DisplayClusterProjection/Private/Policy/VIOSO/Windows/DisplayClusterProjectionVIOSOPolicyViewData.cpp
@@ -277,7 +277,7 @@ bool FDisplayClusterProjectionVIOSOPolicyViewData::InitializeVIOSO_RenderThread(
 	{
 		if (ID3D11DynamicRHI* D3D11DynamicRHI = GetID3D11DynamicRHI())
 		{
-			if (ID3D11Device* D3D11Device = static_cast<ID3D11Device*>(D3D11DynamicRHI->RHIGetNativeDevice()))
+			if (ID3D11Device* D3D11Device = D3D11DynamicRHI->RHIGetDevice())
 			{
 				FD3D11ContextHelper D3D11ContextHelper;
 				if (D3D11ContextHelper.AssignD3D11RenderTarget(RenderTargetTexture) && WarperInterface->Initialize(D3D11Device, ClippingPlanes))
diff --git a/Engine/Plugins/Runtime/nDisplay/Source/DisplayClusterStageMonitoring/DisplayClusterStageMonitoring.Build.cs b/Engine/Plugins/Runtime/nDisplay/Source/DisplayClusterStageMonitoring/DisplayClusterStageMonitoring.Build.cs
--- a/Engine/Plugins/Runtime/nDisplay/Source/DisplayClusterStageMonitoring/DisplayClusterStageMonitoring.Build.cs
+++ b/Engine/Plugins/Runtime/nDisplay/Source/DisplayClusterStageMonitoring/DisplayClusterStageMonitoring.Build.cs
@@ -20,6 +20,11 @@ public class DisplayClusterStageMonitoring : ModuleRules
 				"StageDataCore",
 			});
 
+		PrivateIncludePathModuleNames.AddRange([
+			"D3D11RHI",
+			"D3D12RHI",
+		]);
+
 		AddEngineThirdPartyPrivateStaticDependencies(Target, "NVAPI");
 	}
 }
diff --git a/Engine/Plugins/Runtime/nDisplay/Source/DisplayClusterStageMonitoring/Private/NvidiaSyncWatchdog.cpp b/Engine/Plugins/Runtime/nDisplay/Source/DisplayClusterStageMonitoring/Private/NvidiaSyncWatchdog.cpp
--- a/Engine/Plugins/Runtime/nDisplay/Source/DisplayClusterStageMonitoring/Private/NvidiaSyncWatchdog.cpp
+++ b/Engine/Plugins/Runtime/nDisplay/Source/DisplayClusterStageMonitoring/Private/NvidiaSyncWatchdog.cpp
@@ -11,6 +11,8 @@
 #include "Render/Presentation/IDisplayClusterPresentation.h"
 #include "RendererInterface.h"
 
+#include "ID3D11DynamicRHI.h"
+#include "ID3D12DynamicRHI.h"
 
 #pragma warning(push)
 #pragma warning (disable : 4005) 	// Disable macro redefinition warning for compatibility with Windows SDK 8+
@@ -47,7 +49,13 @@ FNvidiaSyncWatchdog::~FNvidiaSyncWatchdog()
 
 void FNvidiaSyncWatchdog::OnCustomPresentCreated()
 {
-	D3DDevice = static_cast<IUnknown*>(GDynamicRHI->RHIGetNativeDevice());
+	D3DDevice = nullptr;
+	switch (RHIGetInterfaceType())
+	{
+	case ERHIInterfaceType::D3D11: D3DDevice = GetID3D11DynamicRHI()->RHIGetDevice(); break;
+	case ERHIInterfaceType::D3D12: D3DDevice = GetID3D12DynamicRHI()->RHIGetDevice_NoMGPU(); break;
+	}
+
 	if (D3DDevice)
 	{
 		UE_LOG(LogDisplayClusterStageMonitoring, VeryVerbose, TEXT("Nvidia Sync watchdog active"));
diff --git a/Engine/Plugins/Runtime/nDisplay/Source/SharedMemoryMedia/Private/Platforms/Windows/SharedMemoryMediaPlatformWindowsD3D12.cpp b/Engine/Plugins/Runtime/nDisplay/Source/SharedMemoryMedia/Private/Platforms/Windows/SharedMemoryMediaPlatformWindowsD3D12.cpp
--- a/Engine/Plugins/Runtime/nDisplay/Source/SharedMemoryMedia/Private/Platforms/Windows/SharedMemoryMediaPlatformWindowsD3D12.cpp
+++ b/Engine/Plugins/Runtime/nDisplay/Source/SharedMemoryMedia/Private/Platforms/Windows/SharedMemoryMediaPlatformWindowsD3D12.cpp
@@ -99,7 +99,7 @@ FTextureRHIRef FSharedMemoryMediaPlatformWindowsD3D12::CreateSharedTexture(EPixe
 	check(SharedHandle[BufferIdx] == INVALID_HANDLE_VALUE);
 	check(CommittedResource[BufferIdx] == nullptr);
 
-	ID3D12Device* D3D12Device = static_cast<ID3D12Device*>(GDynamicRHI->RHIGetNativeDevice());
+	ID3D12Device* D3D12Device = GetID3D12DynamicRHI()->RHIGetDevice_NoMGPU();
 	check(D3D12Device);
 
 	const DXGI_FORMAT DxgiFormat = FindSharedResourceDXGIFormat(DXGI_FORMAT(GPixelFormats[Format].PlatformFormat), bSrgb);
@@ -199,7 +199,7 @@ FTextureRHIRef FSharedMemoryMediaPlatformWindowsD3D12::OpenSharedTextureByGuid(c
 {
 	using namespace UE::FSharedMemoryMediaPlatformWindowsD3D12;
 
-	ID3D12Device* D3D12Device = static_cast<ID3D12Device*>(GDynamicRHI->RHIGetNativeDevice());
+	ID3D12Device* D3D12Device = GetID3D12DynamicRHI()->RHIGetDevice_NoMGPU();
 	check(D3D12Device);
 
 	const FString SharedGpuTextureName = FString::Printf(TEXT("Global\\%s"), *Guid.ToString(EGuidFormats::DigitsWithHyphensInBraces));
diff --git a/Engine/Plugins/VirtualProduction/TextureShare/Source/TextureShare/Private/Resources/TextureShareResource_D3D12.cpp b/Engine/Plugins/VirtualProduction/TextureShare/Source/TextureShare/Private/Resources/TextureShareResource_D3D12.cpp
--- a/Engine/Plugins/VirtualProduction/TextureShare/Source/TextureShare/Private/Resources/TextureShareResource_D3D12.cpp
+++ b/Engine/Plugins/VirtualProduction/TextureShare/Source/TextureShare/Private/Resources/TextureShareResource_D3D12.cpp
@@ -91,7 +91,7 @@ void FTextureShareResource::InitDynamicRHI_D3D12(FRHICommandListBase& RHICmdList
 			TSharedPtr<ITextureShareCoreD3D12ResourcesCache, ESPMode::ThreadSafe> D3D12ResourcesCache = GetD3D12ResourcesCache();
 			if (D3D12ResourcesCache.IsValid())
 			{
-				if (ID3D12Device* D3D12DevicePtr = GetID3D12DynamicRHI()->RHIGetDevice(0))
+				if (ID3D12Device* D3D12DevicePtr = GetID3D12DynamicRHI()->RHIGetDevice_NoMGPU())
 				{
 					const DXGI_FORMAT PlatformResourceFormat = FindSharedResourceDXGIFormat((DXGI_FORMAT)GPixelFormats[ResourceSettings.Format].PlatformFormat, ResourceSettings.bShouldUseSRGB);
 					FTextureShareCoreResourceHandle ResourceHandle;
@@ -117,7 +117,7 @@ bool FTextureShareResource::D3D12RegisterResourceHandle_RenderThread(FRHICommand
 	if (IsRHID3D12() && TextureRHI.IsValid())
 	{
 		ID3D12DynamicRHI* D3D12RHI = GetID3D12DynamicRHI();
-		ID3D12Device* D3D12DevicePtr = D3D12RHI->RHIGetDevice(0);
+		ID3D12Device* D3D12DevicePtr = D3D12RHI->RHIGetDevice_NoMGPU();
 		ID3D12Resource* D3D12ResourcePtr = (ID3D12Resource*)D3D12RHI->RHIGetResource(TextureRHI);
 
 		if (D3D12DevicePtr && D3D12ResourcePtr)
@@ -165,7 +165,7 @@ bool FTextureShareResource::D3D12ReleaseTextureShareHandle_RenderThread()
 	if (IsRHID3D12() && TextureRHI.IsValid())
 	{
 		ID3D12DynamicRHI* D3D12RHI = GetID3D12DynamicRHI();
-		ID3D12Device* D3D12DevicePtr = D3D12RHI->RHIGetDevice(0);
+		ID3D12Device* D3D12DevicePtr = D3D12RHI->RHIGetDevice_NoMGPU();
 		ID3D12Resource* D3D12ResourcePtr = (ID3D12Resource*)D3D12RHI->RHIGetResource(TextureRHI);
 
 		if (D3D12DevicePtr && D3D12ResourcePtr)
diff --git a/Engine/Source/Runtime/AVEncoder/AVEncoder.Build.cs b/Engine/Source/Runtime/AVEncoder/AVEncoder.Build.cs
--- a/Engine/Source/Runtime/AVEncoder/AVEncoder.Build.cs
+++ b/Engine/Source/Runtime/AVEncoder/AVEncoder.Build.cs
@@ -33,15 +33,15 @@ public class AVEncoder : ModuleRules
 
 		if (Target.IsInPlatformGroup(UnrealPlatformGroup.Windows) || Target.IsInPlatformGroup(UnrealPlatformGroup.Unix))
 		{
-			PublicIncludePathModuleNames.Add("Vulkan");
+			PublicIncludePathModuleNames.Add("VulkanRHI");
 		}
 
 		if (Target.IsInPlatformGroup(UnrealPlatformGroup.Windows))
 		{
-			PrivateDependencyModuleNames.AddRange(new string[] {
+			PublicIncludePathModuleNames.AddRange([
 				"D3D11RHI",
 				"D3D12RHI"
-			});
+			]);
 
 			AddEngineThirdPartyPrivateStaticDependencies(Target, "DX11", "DX12");
 
diff --git a/Engine/Source/Runtime/AVEncoder/Private/MicrosoftCommon.cpp b/Engine/Source/Runtime/AVEncoder/Private/MicrosoftCommon.cpp
--- a/Engine/Source/Runtime/AVEncoder/Private/MicrosoftCommon.cpp
+++ b/Engine/Source/Runtime/AVEncoder/Private/MicrosoftCommon.cpp
@@ -10,24 +10,4 @@
 	#pragma comment(lib, "Mfreadwrite")
 #endif
 
-namespace AVEncoder
-{
-
-//
-// Windows only code
-//
-//#if PLATFORM_WINDOWS
-//ID3D11Device* GetUE4DxDevice()
-//{
-//	auto Device = static_cast<ID3D11Device*>(GDynamicRHI->RHIGetNativeDevice());
-//	checkf(Device != nullptr, TEXT("Failed to get UE4's ID3D11Device"));
-//	return Device;
-//}
-//#endif
-//
-
-
-} // namespace AVEncoder
-
 #endif // PLATFORM_WINDOWS
-
diff --git a/Engine/Source/Runtime/AVEncoder/Private/VideoEncoderFactory.cpp b/Engine/Source/Runtime/AVEncoder/Private/VideoEncoderFactory.cpp
--- a/Engine/Source/Runtime/AVEncoder/Private/VideoEncoderFactory.cpp
+++ b/Engine/Source/Runtime/AVEncoder/Private/VideoEncoderFactory.cpp
@@ -4,6 +4,15 @@
 #include "VideoEncoderInputImpl.h"
 #include "RHI.h"
 
+#if PLATFORM_WINDOWS
+#include "ID3D11DynamicRHI.h"
+#include "ID3D12DynamicRHI.h"
+#endif
+
+#if PLATFORM_DESKTOP || PLATFORM_UNIX
+#include "IVulkanDynamicRHI.h"
+#endif
+
 #include "Encoders/VideoEncoderH264_Dummy.h"
 
 namespace AVEncoder
@@ -119,54 +128,54 @@ PRAGMA_DISABLE_DEPRECATION_WARNINGS
 TUniquePtr<FVideoEncoder> FVideoEncoderFactory::Create(uint32 InID, const FVideoEncoder::FLayerConfig& config)
 {
 	// HACK (M84FIX) create encoder without a ready FVideoEncoderInput
-	TUniquePtr<FVideoEncoder>	Result;
 	for (int32 Index = 0; Index < AvailableEncoders.Num(); ++Index)
 	{
 		if (AvailableEncoders[Index].ID == InID)
 		{
-			Result = CreateEncoders[Index]();
+			TUniquePtr<FVideoEncoder> Result = CreateEncoders[Index]();
+			TSharedPtr<FVideoEncoderInput> Input;
 
-			ERHIInterfaceType RHIType = RHIGetInterfaceType();
+			switch (RHIGetInterfaceType())
+			{
+		#if PLATFORM_WINDOWS
+			case ERHIInterfaceType::D3D11:
+				Input = FVideoEncoderInput::CreateForD3D11(GetID3D11DynamicRHI()->RHIGetDevice(), true, IsRHIDeviceAMD()).ToSharedRef();
+				break;
 
-			if (RHIType == ERHIInterfaceType::D3D11)
-			{	
-				TSharedRef<FVideoEncoderInput> Input = FVideoEncoderInput::CreateForD3D11(GDynamicRHI->RHIGetNativeDevice(), true, IsRHIDeviceAMD()).ToSharedRef();
-				
-				if (Result && !Result->Setup(Input, config))
-				{
-					Result.Reset();
-				}
+			case ERHIInterfaceType::D3D12:
+				Input = FVideoEncoderInput::CreateForD3D12(GetID3D12DynamicRHI()->RHIGetDevice_NoMGPU(), true, IsRHIDeviceNVIDIA()).ToSharedRef();
 				break;
-			}
-			else if (RHIType == ERHIInterfaceType::D3D12)
-			{				
-				TSharedRef<FVideoEncoderInput> Input = FVideoEncoderInput::CreateForD3D12(GDynamicRHI->RHIGetNativeDevice(), true, IsRHIDeviceNVIDIA()).ToSharedRef();
-				
-				if (Result && !Result->Setup(Input, config))
+		#endif // PLATFORM_WINDOWS
+
+		#if PLATFORM_DESKTOP || PLATFORM_UNIX
+			case ERHIInterfaceType::Vulkan:
 				{
-					Result.Reset();
+					IVulkanDynamicRHI* VulkanRHI = GetIVulkanDynamicRHI();
+
+					AVEncoder::FVulkanDataStruct VulkanData =
+					{
+						VulkanRHI->RHIGetVkInstance(),
+						VulkanRHI->RHIGetVkPhysicalDevice(),
+						VulkanRHI->RHIGetVkDevice()
+					};
+
+					Input = FVideoEncoderInput::CreateForVulkan(&VulkanData, true).ToSharedRef();
 				}
 				break;
+		#endif // PLATFORM_DESKTOP || PLATFORM_UNIX
+
+			default:
+				continue;
 			}
-#if PLATFORM_DESKTOP && !PLATFORM_APPLE
-			else if (RHIType == ERHIInterfaceType::Vulkan)
-			{
-				AVEncoder::FVulkanDataStruct VulkanData = {	static_cast<VkInstance>(GDynamicRHI->RHIGetNativeInstance()), 
-															static_cast<VkPhysicalDevice>(GDynamicRHI->RHIGetNativePhysicalDevice()), 
-															static_cast<VkDevice>(GDynamicRHI->RHIGetNativeDevice())};
 
-				TSharedRef<FVideoEncoderInput> Input = FVideoEncoderInput::CreateForVulkan( &VulkanData, true).ToSharedRef();
-				
-				if (Result && !Result->Setup(Input, config))
-				{
-					Result.Reset();
-				}
-				break;
+			if (Result && Input && Result->Setup(Input.ToSharedRef(), config))
+			{
+				return Result;
 			}
-#endif
 		}
 	}
-	return Result;
+
+	return nullptr;
 }
 PRAGMA_ENABLE_DEPRECATION_WARNINGS
 
diff --git a/Engine/Source/Runtime/Apple/MetalRHI/MetalRHI.Build.cs b/Engine/Source/Runtime/Apple/MetalRHI/MetalRHI.Build.cs
--- a/Engine/Source/Runtime/Apple/MetalRHI/MetalRHI.Build.cs
+++ b/Engine/Source/Runtime/Apple/MetalRHI/MetalRHI.Build.cs
@@ -20,6 +20,8 @@ public class MetalRHI : ModuleRules
 			}
 			);
 
+		PublicIncludePathModuleNames.Add("MetalCPP");
+
 		AddEngineThirdPartyPrivateStaticDependencies(Target,
 			"MetalCPP"
 		);
diff --git a/Engine/Source/Runtime/Apple/MetalRHI/Private/MetalDynamicRHI.h b/Engine/Source/Runtime/Apple/MetalRHI/Private/MetalDynamicRHI.h
--- a/Engine/Source/Runtime/Apple/MetalRHI/Private/MetalDynamicRHI.h
+++ b/Engine/Source/Runtime/Apple/MetalRHI/Private/MetalDynamicRHI.h
@@ -185,6 +185,7 @@ public:
 	virtual void RHIBlockUntilGPUIdle() final override;
 	virtual bool RHIGetAvailableResolutions(FScreenResolutionArray& Resolutions, bool bIgnoreRefreshRate) final override;
 	virtual void RHIGetSupportedResolution(uint32& Width, uint32& Height) final override;
+	virtual MTL::Device* RHIGetDevice() final override;
 	virtual void* METALRHI_API RHIGetNativeDevice() final override;
 	virtual void* METALRHI_API RHIGetNativeGraphicsQueue() final override;
 	virtual void* METALRHI_API RHIGetNativeComputeQueue() final override;
diff --git a/Engine/Source/Runtime/Apple/MetalRHI/Private/MetalRHI.cpp b/Engine/Source/Runtime/Apple/MetalRHI/Private/MetalRHI.cpp
--- a/Engine/Source/Runtime/Apple/MetalRHI/Private/MetalRHI.cpp
+++ b/Engine/Source/Runtime/Apple/MetalRHI/Private/MetalRHI.cpp
@@ -1547,6 +1547,11 @@ void* FMetalDynamicRHI::RHIGetNativeDevice()
 	return (void*)Device->GetDevice();
 }
 
+MTL::Device* FMetalDynamicRHI::RHIGetDevice()
+{
+	return Device->GetDevice();
+}
+
 void* FMetalDynamicRHI::RHIGetNativeGraphicsQueue()
 {
 	return ImmediateContext.GetCommandQueue().GetQueue();
diff --git a/Engine/Source/Runtime/Apple/MetalRHI/Public/IMetalDynamicRHI.h b/Engine/Source/Runtime/Apple/MetalRHI/Public/IMetalDynamicRHI.h
--- a/Engine/Source/Runtime/Apple/MetalRHI/Public/IMetalDynamicRHI.h
+++ b/Engine/Source/Runtime/Apple/MetalRHI/Public/IMetalDynamicRHI.h
@@ -12,6 +12,8 @@ struct IMetalDynamicRHI : public FDynamicRHI
 		return ERHIInterfaceType::Metal;
 	}
 	
+	virtual MTL::Device*      RHIGetDevice() = 0;
+
 	virtual void 			  RHIRunOnQueue(TFunction<void(MTL::CommandQueue*)>&& CodeToRun, bool bWaitForSubmission) = 0;
 	virtual FTextureRHIRef	  RHICreateTexture2DFromCVMetalTexture(EPixelFormat Format, ETextureCreateFlags TexCreateFlags, const FClearValueBinding& ClearValueBinding, CVMetalTextureRef Resource) = 0;
 };
diff --git a/Engine/Source/Runtime/CUDA/Source/CUDA.Build.cs b/Engine/Source/Runtime/CUDA/Source/CUDA.Build.cs
--- a/Engine/Source/Runtime/CUDA/Source/CUDA.Build.cs
+++ b/Engine/Source/Runtime/CUDA/Source/CUDA.Build.cs
@@ -9,31 +9,25 @@ public class CUDA : ModuleRules
 {
 	public CUDA(ReadOnlyTargetRules Target) : base(Target)
 	{
-		PrivateDependencyModuleNames.AddRange(new string[] {
+		PrivateDependencyModuleNames.AddRange([
 			"Core",
 			"RenderCore",
 			"RHI",
 			"Engine",
-		});
+		]);
 		
 		PublicDependencyModuleNames.Add("CUDAHeader");
 
 		if (Target.IsInPlatformGroup(UnrealPlatformGroup.Windows) || Target.IsInPlatformGroup(UnrealPlatformGroup.Unix))
 		{
 			PublicDefinitions.Add("PLATFORM_SUPPORTS_CUDA=1");
-			var EngineDir = Path.GetFullPath(Target.RelativeEnginePath);
 
 			PrivateIncludePathModuleNames.Add("VulkanRHI");
 
-			if (Target.IsInPlatformGroup(UnrealPlatformGroup.Unix))
+			if (Target.IsInPlatformGroup(UnrealPlatformGroup.Windows))
 			{
+				PrivateIncludePathModuleNames.AddRange(["D3D11RHI", "D3D12RHI"]);
 			}
-			else if (Target.IsInPlatformGroup(UnrealPlatformGroup.Windows))
-			{
-				AddEngineThirdPartyPrivateStaticDependencies(Target, "DX11", "DX12");
-			}
-
-			AddEngineThirdPartyPrivateStaticDependencies(Target, "Vulkan");
 		}
 		else
 		{
diff --git a/Engine/Source/Runtime/CUDA/Source/Private/CudaModule.cpp b/Engine/Source/Runtime/CUDA/Source/Private/CudaModule.cpp
--- a/Engine/Source/Runtime/CUDA/Source/Private/CudaModule.cpp
+++ b/Engine/Source/Runtime/CUDA/Source/Private/CudaModule.cpp
@@ -16,6 +16,9 @@
 	#include "d3d12.h"
 	#include "d3d11.h"
 	#include "Windows/HideWindowsPlatformTypes.h"
+
+	#include "ID3D11DynamicRHI.h"
+	#include "ID3D12DynamicRHI.h"
 #endif
 
 #include "HAL/UnrealMemory.h"
@@ -161,7 +164,7 @@ bool FCUDAModule::IsRHISelectedDevice(CUdevice cuDevice)
 				return false;
 			}
 
-			ID3D11Device* NativeD3D11Device = static_cast<ID3D11Device*>(GDynamicRHI->RHIGetNativeDevice());
+			ID3D11Device* NativeD3D11Device = GetID3D11DynamicRHI()->RHIGetDevice();
 
 			TRefCountPtr<IDXGIDevice> DXGIDevice;
 
@@ -202,7 +205,7 @@ bool FCUDAModule::IsRHISelectedDevice(CUdevice cuDevice)
 				return false;
 			}
 
-			ID3D12Device* NativeD3D12Device = static_cast<ID3D12Device*>(GDynamicRHI->RHIGetNativeDevice());
+			ID3D12Device* NativeD3D12Device = GetID3D12DynamicRHI()->RHIGetDevice_NoMGPU();
 
 			LUID AdapterLUID = NativeD3D12Device->GetAdapterLuid();
 
diff --git a/Engine/Source/Runtime/D3D12RHI/D3D12RHI.Build.cs b/Engine/Source/Runtime/D3D12RHI/D3D12RHI.Build.cs
--- a/Engine/Source/Runtime/D3D12RHI/D3D12RHI.Build.cs
+++ b/Engine/Source/Runtime/D3D12RHI/D3D12RHI.Build.cs
@@ -10,6 +10,10 @@ public class D3D12RHI : ModuleRules
 	{
 		bAllowConfidentialPlatformDefines = true;
 
+		PrivateIncludePathModuleNames.AddAll(
+			"Shaders"
+		);
+
 		PrivateDependencyModuleNames.AddAll(
 			"CoreUObject",
 			"Engine",
@@ -17,9 +21,9 @@ public class D3D12RHI : ModuleRules
 			"RenderCore",
 			"TraceLog"
 		);
-
-		PrivateIncludePathModuleNames.AddAll(
-			"Shaders"
+		
+		PublicIncludePathModuleNames.AddAll(
+			"DX12"
 		);
 
 		PublicDependencyModuleNames.AddAll(
diff --git a/Engine/Source/Runtime/D3D12RHI/Private/D3D12RHI.cpp b/Engine/Source/Runtime/D3D12RHI/Private/D3D12RHI.cpp
--- a/Engine/Source/Runtime/D3D12RHI/Private/D3D12RHI.cpp
+++ b/Engine/Source/Runtime/D3D12RHI/Private/D3D12RHI.cpp
@@ -665,6 +665,16 @@ ID3D12CommandQueue* FD3D12DynamicRHI::RHIGetCommandQueue() const
 	return GetAdapter().GetDevice(0)->GetQueue(ED3D12QueueType::Direct).D3DCommandQueue;
 }
 
+ID3D12Device* FD3D12DynamicRHI::RHIGetDevice_NoMGPU() const
+{
+	ensureMsgf(GNumExplicitGPUsForRendering == 1,
+		TEXT("Use of RHIGetDevice_NoMGPU() while multi-GPU is enabled. ")
+		TEXT("The calling code may not be handling multi-GPU scenarios correctly. ")
+		TEXT("Switch to using RHIGetDevice(N) where N is the GPU index if the calling code supports MGPU.")
+	);
+	return GetAdapter().GetDevice(0)->GetDevice();
+}
+
 ID3D12Device* FD3D12DynamicRHI::RHIGetDevice(uint32 InIndex) const
 {
 	return GetAdapter().GetDevice(InIndex)->GetDevice();
diff --git a/Engine/Source/Runtime/D3D12RHI/Private/D3D12RHIPrivate.h b/Engine/Source/Runtime/D3D12RHI/Private/D3D12RHIPrivate.h
--- a/Engine/Source/Runtime/D3D12RHI/Private/D3D12RHIPrivate.h
+++ b/Engine/Source/Runtime/D3D12RHI/Private/D3D12RHIPrivate.h
@@ -461,6 +461,7 @@ public:
 	virtual TArray<FD3D12MinimalAdapterDesc> RHIGetAdapterDescs() const final override;
 	virtual bool RHIIsPixEnabled() const final override;
 	virtual ID3D12CommandQueue* RHIGetCommandQueue() const final override;
+	virtual ID3D12Device* RHIGetDevice_NoMGPU() const final override;
 	virtual ID3D12Device* RHIGetDevice(uint32 InIndex) const final override;
 	virtual uint32 RHIGetDeviceNodeMask(uint32 InIndex) const final override;
 	virtual ID3D12GraphicsCommandList* RHIGetGraphicsCommandList(FRHICommandListBase& ExecutingCmdList, uint32 InDeviceIndex) const final override;
diff --git a/Engine/Source/Runtime/D3D12RHI/Public/ID3D12DynamicRHI.h b/Engine/Source/Runtime/D3D12RHI/Public/ID3D12DynamicRHI.h
--- a/Engine/Source/Runtime/D3D12RHI/Public/ID3D12DynamicRHI.h
+++ b/Engine/Source/Runtime/D3D12RHI/Public/ID3D12DynamicRHI.h
@@ -25,6 +25,9 @@ struct ID3D12DynamicRHI : public FDynamicRHI
 	virtual TArray<FD3D12MinimalAdapterDesc> RHIGetAdapterDescs() const = 0;
 	virtual bool                             RHIIsPixEnabled() const = 0;
 
+	// Provided for backwards compatibility. Returns the 0th device, equivalent to RHIGetDevice(0), but also triggers an ensure if MGPU is enabled.
+	virtual ID3D12Device*              RHIGetDevice_NoMGPU() const = 0;
+
 	virtual ID3D12CommandQueue*        RHIGetCommandQueue() const = 0;
 	virtual ID3D12Device*              RHIGetDevice(uint32 InIndex) const = 0;
 	virtual uint32                     RHIGetDeviceNodeMask(uint32 InIndex) const = 0;
diff --git a/Engine/Source/Runtime/GameplayMediaEncoder/GameplayMediaEncoder.Build.cs b/Engine/Source/Runtime/GameplayMediaEncoder/GameplayMediaEncoder.Build.cs
--- a/Engine/Source/Runtime/GameplayMediaEncoder/GameplayMediaEncoder.Build.cs
+++ b/Engine/Source/Runtime/GameplayMediaEncoder/GameplayMediaEncoder.Build.cs
@@ -29,8 +29,7 @@ public class GameplayMediaEncoder : ModuleRules
 
 		if (Target.IsInPlatformGroup(UnrealPlatformGroup.Windows))
 		{
-			PrivateDependencyModuleNames.Add("D3D11RHI");
-			AddEngineThirdPartyPrivateStaticDependencies(Target, "DX11");
+			PrivateIncludePathModuleNames.AddRange(["D3D11RHI", "D3D12RHI", "VulkanRHI"]);
 
 			PublicDelayLoadDLLs.Add("mfplat.dll");
 			PublicDelayLoadDLLs.Add("mfuuid.dll");
diff --git a/Engine/Source/Runtime/GameplayMediaEncoder/Private/GameplayMediaEncoder.cpp b/Engine/Source/Runtime/GameplayMediaEncoder/Private/GameplayMediaEncoder.cpp
--- a/Engine/Source/Runtime/GameplayMediaEncoder/Private/GameplayMediaEncoder.cpp
+++ b/Engine/Source/Runtime/GameplayMediaEncoder/Private/GameplayMediaEncoder.cpp
@@ -20,6 +20,12 @@
 
 #include "AudioEncoderFactory.h"
 
+#if PLATFORM_DESKTOP && !PLATFORM_APPLE
+#include "ID3D11DynamicRHI.h"
+#include "ID3D12DynamicRHI.h"
+#include "IVulkanDynamicRHI.h"
+#endif
+
 DEFINE_LOG_CATEGORY(GameplayMediaEncoder);
 CSV_DEFINE_CATEGORY(GameplayMediaEncoder, true);
 
@@ -260,24 +266,21 @@ bool FGameplayMediaEncoder::Initialize()
 		if (RHIType == ERHIInterfaceType::D3D11)
 		{
 			PRAGMA_DISABLE_DEPRECATION_WARNINGS
-			VideoEncoderInput = AVEncoder::FVideoEncoderInput::CreateForD3D11(GDynamicRHI->RHIGetNativeDevice(), true, IsRHIDeviceAMD());
+			VideoEncoderInput = AVEncoder::FVideoEncoderInput::CreateForD3D11(GetID3D11DynamicRHI()->RHIGetDevice(), true, IsRHIDeviceAMD());
 			PRAGMA_ENABLE_DEPRECATION_WARNINGS
 		}
-		
 		else if (RHIType == ERHIInterfaceType::D3D12)
 		{
 			PRAGMA_DISABLE_DEPRECATION_WARNINGS
-			VideoEncoderInput = AVEncoder::FVideoEncoderInput::CreateForD3D12(GDynamicRHI->RHIGetNativeDevice(), true, IsRHIDeviceNVIDIA());
+			VideoEncoderInput = AVEncoder::FVideoEncoderInput::CreateForD3D12(GetID3D12DynamicRHI()->RHIGetDevice_NoMGPU(), true, IsRHIDeviceNVIDIA());
 			PRAGMA_ENABLE_DEPRECATION_WARNINGS
 		}
-		
 		else if (RHIType == ERHIInterfaceType::Vulkan)
 		{
 			PRAGMA_DISABLE_DEPRECATION_WARNINGS
-			VideoEncoderInput = AVEncoder::FVideoEncoderInput::CreateForVulkan(GDynamicRHI->RHIGetNativeDevice(), true);
+			VideoEncoderInput = AVEncoder::FVideoEncoderInput::CreateForVulkan(GetIVulkanDynamicRHI()->RHIGetVkDevice(), true);
 			PRAGMA_ENABLE_DEPRECATION_WARNINGS
 		}
-		
 		else
 #endif
 		{
diff --git a/Engine/Source/Runtime/WebBrowser/Private/CEF/CEFWebBrowserWindowRHIHelper.cpp b/Engine/Source/Runtime/WebBrowser/Private/CEF/CEFWebBrowserWindowRHIHelper.cpp
--- a/Engine/Source/Runtime/WebBrowser/Private/CEF/CEFWebBrowserWindowRHIHelper.cpp
+++ b/Engine/Source/Runtime/WebBrowser/Private/CEF/CEFWebBrowserWindowRHIHelper.cpp
@@ -6,12 +6,17 @@
 #if WITH_CEF3
 
 #include "CEF/CEFWebBrowserWindow.h"
+
 #if WITH_ENGINE
-#include "RHI.h"
-#if PLATFORM_WINDOWS
-#include "Slate/SlateTextures.h"
-#include "RenderingThread.h"
-#endif
+	#include "RHI.h"
+
+	#if PLATFORM_WINDOWS
+		#include "Slate/SlateTextures.h"
+		#include "RenderingThread.h"
+		#include "ID3D11DynamicRHI.h"
+		#include "ID3D12DynamicRHI.h"
+	#endif
+
 #endif
 
 #if PLATFORM_WINDOWS
@@ -30,7 +35,7 @@ static TComPtr<IDXGIAdapter> GetDXGIAdapterFromRHI()
 	if (RHIType == ERHIInterfaceType::D3D11)
 	{
 		TComPtr<IDXGIDevice> DXGIDevice;
-		if (HRESULT Hr = static_cast<ID3D11Device*>(GDynamicRHI->RHIGetNativeDevice())->QueryInterface(IID_PPV_ARGS(&DXGIDevice)); FAILED(Hr))
+		if (HRESULT Hr = GetID3D11DynamicRHI()->RHIGetDevice()->QueryInterface(IID_PPV_ARGS(&DXGIDevice)); FAILED(Hr))
 		{
 			UE_LOG(LogWebBrowser, Error, TEXT("GetDXGIAdapter() - - ID3D11Device::QueryInterface 0x%x"), Hr);
 			return DXGIAdapter;
@@ -43,7 +48,7 @@ static TComPtr<IDXGIAdapter> GetDXGIAdapterFromRHI()
 	}
 	else if (RHIType == ERHIInterfaceType::D3D12)
 	{
-		const LUID Luid = static_cast<ID3D12Device*>(GDynamicRHI->RHIGetNativeDevice())->GetAdapterLuid();
+		const LUID Luid = GetID3D12DynamicRHI()->RHIGetDevice_NoMGPU()->GetAdapterLuid();
 
 		TComPtr<IDXGIFactory4> DXGIFactory;
 		if (HRESULT Hr = CreateDXGIFactory1(IID_PPV_ARGS(&DXGIFactory)); FAILED(Hr))
@@ -121,7 +126,7 @@ uint64_t FCEFWebBrowserWindowRHIHelper::GetRHIAdapterLuid()
 	}
 	else if (RHIType == ERHIInterfaceType::D3D12)
 	{
-		Luid = static_cast<ID3D12Device*>(GDynamicRHI->RHIGetNativeDevice())->GetAdapterLuid();
+		Luid = GetID3D12DynamicRHI()->RHIGetDevice_NoMGPU()->GetAdapterLuid();
 	}
 
 	AdapterLuid = (static_cast<uint64_t>(Luid.HighPart) << 32) | static_cast<uint64_t>(Luid.LowPart);
@@ -211,7 +216,7 @@ bool FCEFWebBrowserWindowRHIHelper::EnsureShareable(FSlateUpdatableTexture* Slat
 	}
 	else if (RHIType == ERHIInterfaceType::D3D12)
 	{
-		if (HRESULT Hr = static_cast<ID3D12Device*>(GDynamicRHI->RHIGetNativeDevice())->CreateSharedHandle(static_cast<ID3D12Resource*>(Slate2DRef->GetNativeResource()), NULL, GENERIC_ALL, NULL, &SharedHandle); FAILED(Hr))
+		if (HRESULT Hr = GetID3D12DynamicRHI()->RHIGetDevice_NoMGPU()->CreateSharedHandle(static_cast<ID3D12Resource*>(Slate2DRef->GetNativeResource()), NULL, GENERIC_ALL, NULL, &SharedHandle); FAILED(Hr))
 		{
 			UE_LOG(LogWebBrowser, Error, TEXT("FCEFWebBrowserWindowRHIHelper::EnsureShareable() - - ID3D12Device::CreateSharedHandle 0x%x"), Hr);
 		}
diff --git a/Engine/Source/Runtime/WebBrowser/WebBrowser.Build.cs b/Engine/Source/Runtime/WebBrowser/WebBrowser.Build.cs
--- a/Engine/Source/Runtime/WebBrowser/WebBrowser.Build.cs
+++ b/Engine/Source/Runtime/WebBrowser/WebBrowser.Build.cs
@@ -88,6 +88,11 @@ public class WebBrowser : ModuleRules
 
 		if (Target.Platform == UnrealTargetPlatform.Win64)
 		{
+			PrivateIncludePathModuleNames.AddRange([
+				"D3D11RHI",
+				"D3D12RHI"
+			]);
+
 			AddEngineThirdPartyPrivateStaticDependencies(Target, "DX11");
 		}
 
diff --git a/Engine/Source/Runtime/Windows/D3D11RHI/D3D11RHI.Build.cs b/Engine/Source/Runtime/Windows/D3D11RHI/D3D11RHI.Build.cs
--- a/Engine/Source/Runtime/Windows/D3D11RHI/D3D11RHI.Build.cs
+++ b/Engine/Source/Runtime/Windows/D3D11RHI/D3D11RHI.Build.cs
@@ -8,7 +8,9 @@ public class D3D11RHI : ModuleRules
 {
 	public D3D11RHI(ReadOnlyTargetRules Target) : base(Target)
 	{
-		PrivateIncludePathModuleNames.AddRange(new string[] { "Shaders" });
+		PrivateIncludePathModuleNames.AddAll(
+			"Shaders"
+		);
 
 		PrivateDependencyModuleNames.AddAll(
 			"CoreUObject",
@@ -18,6 +20,10 @@ public class D3D11RHI : ModuleRules
 			"TraceLog"
 		);
 
+		PublicIncludePathModuleNames.AddAll(
+			"DX11"
+		);
+
 		PublicDependencyModuleNames.AddAll(
 			"Core",
 			"RHI"
